[SlateInspectorToolset] Rename `SlateInspector` plugin and module to `SlateInspectorToolset`, rename all internal classes accordingly.

#changelist validated
#Synced-CL 52321477

#Synced-CL 52321477

diff --git a/Engine/Plugins/Experimental/SlateInspector/SlateInspector.uplugin b/Engine/Plugins/Experimental/SlateInspector/SlateInspector.uplugin
deleted file mode 100644
diff --git a/Engine/Plugins/Experimental/SlateInspector/Source/SlateInspector/Private/SlateInspectorModule.cpp b/Engine/Plugins/Experimental/SlateInspector/Source/SlateInspector/Private/SlateInspectorModule.cpp
deleted file mode 100644
diff --git a/Engine/Plugins/Experimental/SlateInspector/Source/SlateInspector/Private/SlateInspectorObserverManager.cpp b/Engine/Plugins/Experimental/SlateInspector/Source/SlateInspector/Private/SlateInspectorObserverManager.cpp
deleted file mode 100644
diff --git a/Engine/Plugins/Experimental/SlateInspector/Source/SlateInspector/Private/SlateInspectorRefCache.cpp b/Engine/Plugins/Experimental/SlateInspector/Source/SlateInspector/Private/SlateInspectorRefCache.cpp
deleted file mode 100644
diff --git a/Engine/Plugins/Experimental/SlateInspector/Source/SlateInspector/Private/SlateInspectorSnapshotRenderer.cpp b/Engine/Plugins/Experimental/SlateInspector/Source/SlateInspector/Private/SlateInspectorSnapshotRenderer.cpp
deleted file mode 100644
diff --git a/Engine/Plugins/Experimental/SlateInspector/Source/SlateInspector/Private/SlateInspectorSubsystem.cpp b/Engine/Plugins/Experimental/SlateInspector/Source/SlateInspector/Private/SlateInspectorSubsystem.cpp
deleted file mode 100644
diff --git a/Engine/Plugins/Experimental/SlateInspector/Source/SlateInspector/Private/SlateInspectorToolset.cpp b/Engine/Plugins/Experimental/SlateInspector/Source/SlateInspector/Private/SlateInspectorToolset.cpp
deleted file mode 100644
diff --git a/Engine/Plugins/Experimental/SlateInspector/Source/SlateInspector/Public/SlateInspectorModule.h b/Engine/Plugins/Experimental/SlateInspector/Source/SlateInspector/Public/SlateInspectorModule.h
deleted file mode 100644
diff --git a/Engine/Plugins/Experimental/SlateInspector/Source/SlateInspector/Public/SlateInspectorObserverManager.h b/Engine/Plugins/Experimental/SlateInspector/Source/SlateInspector/Public/SlateInspectorObserverManager.h
deleted file mode 100644
diff --git a/Engine/Plugins/Experimental/SlateInspector/Source/SlateInspector/Public/SlateInspectorRefCache.h b/Engine/Plugins/Experimental/SlateInspector/Source/SlateInspector/Public/SlateInspectorRefCache.h
deleted file mode 100644
diff --git a/Engine/Plugins/Experimental/SlateInspector/Source/SlateInspector/Public/SlateInspectorSnapshotRenderer.h b/Engine/Plugins/Experimental/SlateInspector/Source/SlateInspector/Public/SlateInspectorSnapshotRenderer.h
deleted file mode 100644
diff --git a/Engine/Plugins/Experimental/SlateInspector/Source/SlateInspector/Public/SlateInspectorSubsystem.h b/Engine/Plugins/Experimental/SlateInspector/Source/SlateInspector/Public/SlateInspectorSubsystem.h
deleted file mode 100644
diff --git a/Engine/Plugins/Experimental/SlateInspector/Source/SlateInspector/Public/SlateInspectorToolset.h b/Engine/Plugins/Experimental/SlateInspector/Source/SlateInspector/Public/SlateInspectorToolset.h
deleted file mode 100644
diff --git a/Engine/Plugins/Experimental/SlateInspector/Source/SlateInspector/SlateInspector.Build.cs b/Engine/Plugins/Experimental/SlateInspector/Source/SlateInspector/SlateInspector.Build.cs
deleted file mode 100644
diff --git a/Engine/Plugins/Experimental/SlateInspector/Source/SlateInspectorTests/Private/SlateInspectorTestPanel.cpp b/Engine/Plugins/Experimental/SlateInspector/Source/SlateInspectorTests/Private/SlateInspectorTestPanel.cpp
deleted file mode 100644
diff --git a/Engine/Plugins/Experimental/SlateInspector/Source/SlateInspectorTests/Private/SlateInspectorTestPanel.h b/Engine/Plugins/Experimental/SlateInspector/Source/SlateInspectorTests/Private/SlateInspectorTestPanel.h
deleted file mode 100644
diff --git a/Engine/Plugins/Experimental/SlateInspector/Source/SlateInspectorTests/Private/SlateInspectorTests.cpp b/Engine/Plugins/Experimental/SlateInspector/Source/SlateInspectorTests/Private/SlateInspectorTests.cpp
deleted file mode 100644
diff --git a/Engine/Plugins/Experimental/SlateInspector/Source/SlateInspectorTests/Private/SlateInspectorTestsModule.cpp b/Engine/Plugins/Experimental/SlateInspector/Source/SlateInspectorTests/Private/SlateInspectorTestsModule.cpp
deleted file mode 100644
diff --git a/Engine/Plugins/Experimental/SlateInspector/Source/SlateInspectorTests/SlateInspectorTests.Build.cs b/Engine/Plugins/Experimental/SlateInspector/Source/SlateInspectorTests/SlateInspectorTests.Build.cs
deleted file mode 100644
diff --git a/Engine/Plugins/Experimental/SlateInspectorToolset/SlateInspectorToolset.uplugin b/Engine/Plugins/Experimental/SlateInspectorToolset/SlateInspectorToolset.uplugin
new file mode 100644
--- /dev/null
+++ b/Engine/Plugins/Experimental/SlateInspectorToolset/SlateInspectorToolset.uplugin
@@ -0,0 +1,36 @@
+{
+	"FileVersion": 3,
+	"Version": 1,
+	"VersionName": "1.0",
+	"FriendlyName": "SlateInspectorToolset",
+	"Description": "Slate UI automation and inspection tools.",
+	"NoRedist": true,
+	"EditorOnly": true,
+	"CreatedBy": "",
+	"CreatedByURL": "",
+	"DocsURL": "",
+	"MarketplaceURL": "",
+	"CanContainContent": false,
+	"IsBetaVersion": false,
+	"IsExperimentalVersion": true,
+	"Installed": false,
+	"EnabledByDefault": false,
+	"Modules": [
+		{
+			"Name": "SlateInspectorToolset",
+			"Type": "Editor",
+			"LoadingPhase": "Default"
+		},
+		{
+			"Name": "SlateInspectorToolsetTests",
+			"Type": "Editor",
+			"LoadingPhase": "Default"
+		}
+	],
+	"Plugins": [
+		{
+			"Name": "ToolsetRegistry",
+			"Enabled": true
+		}
+	]
+}
diff --git a/Engine/Plugins/Experimental/SlateInspectorToolset/Source/SlateInspectorToolset/Private/SlateInspectorToolset.cpp b/Engine/Plugins/Experimental/SlateInspectorToolset/Source/SlateInspectorToolset/Private/SlateInspectorToolset.cpp
new file mode 100644
--- /dev/null
+++ b/Engine/Plugins/Experimental/SlateInspectorToolset/Source/SlateInspectorToolset/Private/SlateInspectorToolset.cpp
@@ -0,0 +1,801 @@
+// Copyright Epic Games, Inc. All Rights Reserved.
+
+#include "SlateInspectorToolset.h"
+
+#include "SlateInspectorToolsetObserverManager.h"
+#include "SlateInspectorToolsetRefCache.h"
+#include "SlateInspectorToolsetSnapshotRenderer.h"
+
+#include "Framework/Application/SlateApplication.h"
+#include "Input/Events.h"
+#include "Logging/StructuredLog.h"
+#include "Widgets/Input/SCheckBox.h"
+#include "Widgets/SWindow.h"
+#include "Widgets/Text/STextBlock.h"
+
+DEFINE_LOG_CATEGORY_STATIC(LogSlateInspectorToolset, Log, All);
+
+namespace SlateInspectorToolsetPrivate
+{
+
+/** Escape a string for safe embedding in a JSON string value (RFC 8259). */
+FString EscapeJsonString(const FString& Input)
+{
+	FString Result;
+	Result.Reserve(Input.Len() + 16);
+	for (TCHAR Character : Input)
+	{
+		switch (Character)
+		{
+		case TEXT('"'):  Result += TEXT("\\\""); break;
+		case TEXT('\\'): Result += TEXT("\\\\"); break;
+		case TEXT('\b'): Result += TEXT("\\b");  break;
+		case TEXT('\f'): Result += TEXT("\\f");  break;
+		case TEXT('\n'): Result += TEXT("\\n");  break;
+		case TEXT('\r'): Result += TEXT("\\r");  break;
+		case TEXT('\t'): Result += TEXT("\\t");  break;
+		default:
+			// Escape remaining ASCII control characters (U+0000..U+001F) per RFC 8259.
+			if (Character < 0x20)
+			{
+				Result += FString::Printf(TEXT("\\u%04x"), static_cast<uint32>(Character));
+			}
+			else
+			{
+				Result += Character;
+			}
+			break;
+		}
+	}
+	return Result;
+}
+
+} // namespace SlateInspectorToolsetPrivate
+
+// ── Helpers ──────────────────────────────────────────────────────────
+
+TSharedPtr<SWidget> USlateInspectorToolset::ResolveRefOrWarn(const FString& Ref)
+{
+	TSharedPtr<SWidget> Widget = FSlateInspectorToolsetRefCache::Get().ResolveRef(Ref);
+	if (!Widget.IsValid())
+	{
+		UE_LOGFMT(LogSlateInspectorToolset, Warning, "SlateInspectorToolset: Could not resolve ref '{Ref}'", Ref);
+	}
+	return Widget;
+}
+
+FVector2D USlateInspectorToolset::GetWidgetScreenCenter(TSharedRef<SWidget> Widget)
+{
+	const FGeometry& Geometry = Widget->GetTickSpaceGeometry();
+	return FVector2D(Geometry.GetAbsolutePositionAtCoordinates(FVector2D(0.5, 0.5)));
+}
+
+FKey USlateInspectorToolset::ParseMouseButtonKey(const FString& Button)
+{
+	if (Button.Equals(TEXT("right"), ESearchCase::IgnoreCase))
+	{
+		return EKeys::RightMouseButton;
+	}
+	if (Button.Equals(TEXT("middle"), ESearchCase::IgnoreCase))
+	{
+		return EKeys::MiddleMouseButton;
+	}
+	return EKeys::LeftMouseButton;
+}
+
+bool USlateInspectorToolset::SimulateClick(TSharedRef<SWidget> Widget, FKey MouseButton, bool bDoubleClick, const FModifierKeysState& ModifierKeys)
+{
+	if (!FSlateApplication::IsInitialized())
+	{
+		return false;
+	}
+
+	// Find the window containing this widget by walking up its parent chain
+	// to the nearest SWindow. FindWidgetWindow can return the wrong window
+	// for some widget types, and LocateWindowUnderMouse depends on z-order.
+	TSharedPtr<SWindow> Window;
+	TSharedPtr<SWidget> Current = Widget;
+	static const FName WindowType(TEXT("SWindow"));
+	while (Current.IsValid())
+	{
+		if (Current->GetType() == WindowType)
+		{
+			Window = StaticCastSharedPtr<SWindow>(Current);
+			break;
+		}
+		Current = Current->GetParentWidget();
+	}
+	if (!Window.IsValid())
+	{
+		UE_LOGFMT(LogSlateInspectorToolset, Warning, "SlateInspectorToolset: Could not find window containing widget");
+		return false;
+	}
+
+	FVector2D ScreenCenter = GetWidgetScreenCenter(Widget);
+	FSlateApplication::Get().SetCursorPos(ScreenCenter);
+
+	// Send a mouse move so Slate updates its hover tracking, widget path
+	// cache, and window z-order for hit-testing.
+	FPointerEvent MoveEvent(
+		0,
+		ScreenCenter,
+		ScreenCenter,
+		TSet<FKey>(),
+		EKeys::Invalid,
+		0.0f,
+		ModifierKeys);
+	FSlateApplication::Get().ProcessMouseMoveEvent(MoveEvent);
+
+	TSet<FKey> PressedButtons;
+	PressedButtons.Add(MouseButton);
+
+	FPointerEvent MouseEvent(
+		0,
+		ScreenCenter,
+		ScreenCenter,
+		PressedButtons,
+		MouseButton,
+		0.0f,
+		ModifierKeys);
+
+	if (bDoubleClick)
+	{
+		FSlateApplication::Get().ProcessMouseButtonDoubleClickEvent(Window->GetNativeWindow(), MouseEvent);
+	}
+	else
+	{
+		FSlateApplication::Get().ProcessMouseButtonDownEvent(Window->GetNativeWindow(), MouseEvent);
+	}
+
+	// Release for both paths. ProcessMouseButtonDoubleClickEvent implies a
+	// mouse-down, so we must release to avoid leaving widgets in pressed state.
+	FPointerEvent MouseUpEvent(
+		0,
+		ScreenCenter,
+		ScreenCenter,
+		TSet<FKey>(),
+		MouseButton,
+		0.0f,
+		ModifierKeys);
+	FSlateApplication::Get().ProcessMouseButtonUpEvent(MouseUpEvent);
+
+	return true;
+}
+
+// ── Snapshot tools ───────────────────────────────────────────────────
+
+FString USlateInspectorToolset::Snapshot(const FString& Ref, int32 MaxDepth, bool bIncludeSourceLocations)
+{
+	MaxDepth = FMath::Max(MaxDepth, 0);
+
+	TSharedPtr<SWidget> Root;
+	if (!Ref.IsEmpty())
+	{
+		Root = ResolveRefOrWarn(Ref);
+		if (!Root.IsValid())
+		{
+			return TEXT("Error: Could not resolve ref");
+		}
+	}
+
+	// If an observer covers this root, return its cached snapshot.
+	if (!bIncludeSourceLocations)
+	{
+		FString CachedSnapshot = FSlateInspectorToolsetObserverManager::Get().FindMatchingObserverSnapshot(Root, MaxDepth);
+		if (!CachedSnapshot.IsEmpty())
+		{
+			return CachedSnapshot;
+		}
+	}
+
+	// No observer covers this request. Render fresh without resetting the cache
+	// (observers may be maintaining refs for other subtrees).
+	return FSlateInspectorToolsetSnapshotRenderer::Render(Root, MaxDepth, bIncludeSourceLocations, /* bResetCache */ false);
+}
+
+FString USlateInspectorToolset::Observe(const FString& Ref, int32 MaxDepth)
+{
+	MaxDepth = FMath::Max(MaxDepth, 0);
+
+	TSharedPtr<SWidget> RootWidget;
+	if (!Ref.IsEmpty())
+	{
+		RootWidget = ResolveRefOrWarn(Ref);
+		if (!RootWidget.IsValid())
+		{
+			return TEXT("Error: Could not resolve ref");
+		}
+	}
+
+	return FSlateInspectorToolsetObserverManager::Get().AddObserver(RootWidget, MaxDepth);
+}
+
+bool USlateInspectorToolset::Unobserve(const FString& Identifier)
+{
+	return FSlateInspectorToolsetObserverManager::Get().RemoveObserver(Identifier);
+}
+
+FString USlateInspectorToolset::ListObservers()
+{
+	TArray<FSlateInspectorToolsetObserver> Observers = FSlateInspectorToolsetObserverManager::Get().GetObservers();
+
+	FString Result = TEXT("[");
+	for (int32 Index = 0; Index < Observers.Num(); ++Index)
+	{
+		const FSlateInspectorToolsetObserver& Observer = Observers[Index];
+		if (Index > 0)
+		{
+			Result += TEXT(", ");
+		}
+
+		FString RootRef;
+		if (!Observer.bRoot)
+		{
+			TSharedPtr<SWidget> RootWidget = Observer.RootWidget.Pin();
+			if (RootWidget.IsValid())
+			{
+				RootRef = FSlateInspectorToolsetRefCache::Get().FindRef(RootWidget.ToSharedRef());
+			}
+		}
+
+		Result += FString::Printf(
+			TEXT("{\"identifier\": \"%s\", \"root\": %s, \"rootRef\": \"%s\", \"maxDepth\": %d, \"cachedSnapshotSize\": %d}"),
+			*SlateInspectorToolsetPrivate::EscapeJsonString(Observer.Identifier),
+			Observer.bRoot ? TEXT("true") : TEXT("false"),
+			*SlateInspectorToolsetPrivate::EscapeJsonString(RootRef),
+			Observer.MaxDepth,
+			Observer.CachedSnapshotText.Len());
+	}
+	Result += TEXT("]");
+
+	return Result;
+}
+
+FToolsetImage USlateInspectorToolset::Screenshot(const FString& Ref)
+{
+	FToolsetImage Result;
+
+	if (!FSlateApplication::IsInitialized())
+	{
+		return Result;
+	}
+
+	TSharedPtr<SWidget> Widget;
+	if (!Ref.IsEmpty())
+	{
+		Widget = ResolveRefOrWarn(Ref);
+	}
+	else
+	{
+		// Use the active top-level window.
+		TSharedPtr<SWindow> ActiveWindow = FSlateApplication::Get().GetActiveTopLevelWindow();
+		Widget = ActiveWindow;
+	}
+
+	if (!Widget.IsValid())
+	{
+		return Result;
+	}
+
+	TArray<FColor> ColorData;
+	FIntVector Size;
+	if (FSlateApplication::Get().TakeScreenshot(Widget.ToSharedRef(), ColorData, Size))
+	{
+		Result.SetFromBitmap(ColorData, FIntPoint(Size.X, Size.Y));
+	}
+
+	return Result;
+}
+
+// ── Action tools ─────────────────────────────────────────────────────
+
+bool USlateInspectorToolset::Click(const FString& Ref, const FString& Button, bool DoubleClick, const FSlateInspectorToolsetModifierKeys& Modifiers)
+{
+	if (!FSlateApplication::IsInitialized())
+	{
+		return false;
+	}
+
+	TSharedPtr<SWidget> Widget = ResolveRefOrWarn(Ref);
+	if (!Widget.IsValid())
+	{
+		return false;
+	}
+
+	return SimulateClick(Widget.ToSharedRef(), ParseMouseButtonKey(Button), DoubleClick, Modifiers.ToModifierKeysState());
+}
+
+bool USlateInspectorToolset::Hover(const FString& Ref)
+{
+	TSharedPtr<SWidget> Widget = ResolveRefOrWarn(Ref);
+	if (!Widget.IsValid() || !FSlateApplication::IsInitialized())
+	{
+		return false;
+	}
+
+	FVector2D ScreenCenter = GetWidgetScreenCenter(Widget.ToSharedRef());
+	FSlateApplication::Get().SetCursorPos(ScreenCenter);
+
+	// SetCursorPos alone does not update Slate's hover tracking.
+	// Send a mouse move event so the widget receives OnMouseEnter and IsHovered() becomes true.
+	FModifierKeysState ModifierKeys;
+	FPointerEvent MouseMoveEvent(
+		0, ScreenCenter, ScreenCenter,
+		TSet<FKey>(), FKey(), 0.0f, ModifierKeys);
+	FSlateApplication::Get().ProcessMouseMoveEvent(MouseMoveEvent);
+
+	return true;
+}
+
+bool USlateInspectorToolset::Type(const FString& Ref, const FString& Text, bool Submit)
+{
+	if (!FSlateApplication::IsInitialized())
+	{
+		return false;
+	}
+
+	TSharedPtr<SWidget> Widget = ResolveRefOrWarn(Ref);
+	if (!Widget.IsValid())
+	{
+		return false;
+	}
+
+	// Focus the target widget.
+	FSlateApplication::Get().SetKeyboardFocus(Widget, EFocusCause::SetDirectly);
+
+	if (!Widget->HasKeyboardFocus())
+	{
+		UE_LOGFMT(LogSlateInspectorToolset, Warning, "SlateInspectorToolset: Could not focus widget for typing");
+		return false;
+	}
+
+	FModifierKeysState ModifierKeys;
+
+	// Send each character as a key char event.
+	for (TCHAR Character : Text)
+	{
+		FCharacterEvent CharEvent(Character, ModifierKeys, 0, false);
+		FSlateApplication::Get().ProcessKeyCharEvent(CharEvent);
+	}
+
+	if (Submit)
+	{
+		FKeyEvent KeyDownEvent(EKeys::Enter, ModifierKeys, 0, false, 0, 0);
+		FSlateApplication::Get().ProcessKeyDownEvent(KeyDownEvent);
+		FKeyEvent KeyUpEvent(EKeys::Enter, ModifierKeys, 0, false, 0, 0);
+		FSlateApplication::Get().ProcessKeyUpEvent(KeyUpEvent);
+	}
+
+	return true;
+}
+
+bool USlateInspectorToolset::PressKey(const FString& Key)
+{
+	if (!FSlateApplication::IsInitialized())
+	{
+		return false;
+	}
+
+	// Parse modifier prefixes: "Ctrl+Shift+Alt+Meta+KeyName"
+	bool bShift = false;
+	bool bControl = false;
+	bool bAlt = false;
+	bool bCommand = false;
+	FString Remaining = Key;
+
+	while (true)
+	{
+		if (Remaining.StartsWith(TEXT("Shift+"), ESearchCase::IgnoreCase))
+		{
+			bShift = true;
+			Remaining.RightChopInline(6);
+		}
+		else if (Remaining.StartsWith(TEXT("Ctrl+"), ESearchCase::IgnoreCase) || Remaining.StartsWith(TEXT("Control+"), ESearchCase::IgnoreCase))
+		{
+			bControl = true;
+			Remaining.RightChopInline(Remaining.StartsWith(TEXT("Ctrl+"), ESearchCase::IgnoreCase) ? 5 : 8);
+		}
+		else if (Remaining.StartsWith(TEXT("Alt+"), ESearchCase::IgnoreCase))
+		{
+			bAlt = true;
+			Remaining.RightChopInline(4);
+		}
+		else if (Remaining.StartsWith(TEXT("Meta+"), ESearchCase::IgnoreCase) || Remaining.StartsWith(TEXT("Cmd+"), ESearchCase::IgnoreCase))
+		{
+			bCommand = true;
+			Remaining.RightChopInline(Remaining.StartsWith(TEXT("Meta+"), ESearchCase::IgnoreCase) ? 5 : 4);
+		}
+		else
+		{
+			break;
+		}
+	}
+
+	FKey ParsedKey(*Remaining);
+
+	if (!ParsedKey.IsValid())
+	{
+		UE_LOGFMT(LogSlateInspectorToolset, Warning, "SlateInspectorToolset: Unknown key '{Key}'", Key);
+		return false;
+	}
+
+	FModifierKeysState ModifierKeys(bShift, bShift, bControl, bControl, bAlt, bAlt, bCommand, bCommand, false);
+	FKeyEvent KeyDownEvent(ParsedKey, ModifierKeys, 0, false, 0, 0);
+	FSlateApplication::Get().ProcessKeyDownEvent(KeyDownEvent);
+
+	// Some keys are handled via OnKeyChar rather than OnKeyDown (e.g.
+	// BackSpace in SEditableText). Normally the OS translates WM_KEYDOWN
+	// into WM_CHAR, but since we bypass the OS message loop we must send
+	// the character event ourselves. Only send for unmodified keys (Ctrl/Alt
+	// combos are handled purely via OnKeyDown).
+	if (!bControl && !bAlt && !bCommand)
+	{
+		static const TMap<FKey, TCHAR> KeyToCharacter = {
+			{ EKeys::BackSpace, TCHAR('\b') },
+			{ EKeys::Tab,       TCHAR('\t') },
+			{ EKeys::Enter,     TCHAR('\r') },
+			{ EKeys::Escape,    TCHAR('\x1b') },
+		};
+
+		if (const TCHAR* CharCode = KeyToCharacter.Find(ParsedKey))
+		{
+			FCharacterEvent CharEvent(*CharCode, ModifierKeys, 0, false);
+			FSlateApplication::Get().ProcessKeyCharEvent(CharEvent);
+		}
+	}
+
+	FKeyEvent KeyUpEvent(ParsedKey, ModifierKeys, 0, false, 0, 0);
+	FSlateApplication::Get().ProcessKeyUpEvent(KeyUpEvent);
+
+	return true;
+}
+
+bool USlateInspectorToolset::SelectOption(const FString& Ref, const FString& Value)
+{
+	if (!FSlateApplication::IsInitialized())
+	{
+		return false;
+	}
+
+	TSharedPtr<SWidget> Widget = ResolveRefOrWarn(Ref);
+	if (!Widget.IsValid())
+	{
+		return false;
+	}
+
+	// Capture the set of windows before opening, so we can identify the new dropdown popup.
+	TArray<TSharedRef<SWindow>> WindowsBefore;
+	FSlateApplication::Get().GetAllVisibleWindowsOrdered(WindowsBefore);
+
+	// Click to open the combobox dropdown.
+	if (!SimulateClick(Widget.ToSharedRef(), EKeys::LeftMouseButton, false))
+	{
+		return false;
+	}
+
+	// Pump Slate twice: first tick creates the popup window, second tick
+	// lays out and materializes the list view items inside it.
+	// Guard against reentrancy if we're already inside a Slate tick.
+	if (!ensure(!FSlateApplication::Get().IsProcessingInput()))
+	{
+		UE_LOGFMT(LogSlateInspectorToolset, Warning, "SlateInspectorToolset: SelectOption called during input processing, cannot tick");
+		return false;
+	}
+	FSlateApplication::Get().Tick();
+	FSlateApplication::Get().Tick();
+
+	static const FName TextBlockType(TEXT("STextBlock"));
+	TSharedPtr<SWidget> MatchingTextBlock;
+	TArray<TSharedRef<SWindow>> VisibleWindows;
+	FSlateApplication::Get().GetAllVisibleWindowsOrdered(VisibleWindows);
+
+	// Search newly appeared windows first (the dropdown popup), then fall back
+	// to all windows. This avoids matching text in the combobox's own label.
+	TSet<TSharedRef<SWindow>> WindowsBeforeSet(WindowsBefore);
+	VisibleWindows.Sort([&WindowsBeforeSet](const TSharedRef<SWindow>& A, const TSharedRef<SWindow>& B)
+	{
+		bool bAIsNew = !WindowsBeforeSet.Contains(A);
+		bool bBIsNew = !WindowsBeforeSet.Contains(B);
+		return bAIsNew > bBIsNew;
+	});
+
+	for (const TSharedRef<SWindow>& Window : VisibleWindows)
+	{
+		TArray<TSharedRef<SWidget>> Queue;
+		FChildren* WindowChildren = Window->GetChildren();
+		for (int32 ChildIndex = 0; ChildIndex < WindowChildren->Num(); ++ChildIndex)
+		{
+			Queue.Add(WindowChildren->GetChildAt(ChildIndex));
+		}
+
+		int32 QueueIndex = 0;
+		while (QueueIndex < Queue.Num())
+		{
+			TSharedRef<SWidget> Current = Queue[QueueIndex++];
+
+			if (Current->GetType() == TextBlockType)
+			{
+				FText WidgetText = StaticCastSharedRef<STextBlock>(Current)->GetText();
+				if (WidgetText.ToString() == Value)
+				{
+					MatchingTextBlock = Current;
+					break;
+				}
+			}
+
+			FChildren* CurrentChildren = Current->GetChildren();
+			for (int32 ChildIndex = 0; ChildIndex < CurrentChildren->Num(); ++ChildIndex)
+			{
+				Queue.Add(CurrentChildren->GetChildAt(ChildIndex));
+			}
+		}
+
+		if (MatchingTextBlock.IsValid())
+		{
+			break;
+		}
+	}
+
+	if (!MatchingTextBlock.IsValid())
+	{
+		UE_LOGFMT(LogSlateInspectorToolset, Warning, "SlateInspectorToolset: Could not find option '{Value}'", Value);
+		return false;
+	}
+
+	return SimulateClick(MatchingTextBlock.ToSharedRef(), EKeys::LeftMouseButton, false);
+}
+
+bool USlateInspectorToolset::Drag(const FString& StartRef, const FString& EndRef, const FSlateInspectorToolsetModifierKeys& Modifiers)
+{
+	if (!FSlateApplication::IsInitialized())
+	{
+		return false;
+	}
+
+	TSharedPtr<SWidget> StartWidget = ResolveRefOrWarn(StartRef);
+	TSharedPtr<SWidget> EndWidget = ResolveRefOrWarn(EndRef);
+	if (!StartWidget.IsValid() || !EndWidget.IsValid())
+	{
+		return false;
+	}
+
+	// Walk parent chain to find the containing SWindow (same approach as SimulateClick).
+	TSharedPtr<SWindow> StartWindow;
+	TSharedPtr<SWidget> WalkCurrent = StartWidget;
+	static const FName WindowType(TEXT("SWindow"));
+	while (WalkCurrent.IsValid())
+	{
+		if (WalkCurrent->GetType() == WindowType)
+		{
+			StartWindow = StaticCastSharedPtr<SWindow>(WalkCurrent);
+			break;
+		}
+		WalkCurrent = WalkCurrent->GetParentWidget();
+	}
+	if (!StartWindow.IsValid())
+	{
+		UE_LOGFMT(LogSlateInspectorToolset, Warning, "SlateInspectorToolset: Could not find window for drag source");
+		return false;
+	}
+
+	FVector2D StartPosition = GetWidgetScreenCenter(StartWidget.ToSharedRef());
+	FVector2D EndPosition = GetWidgetScreenCenter(EndWidget.ToSharedRef());
+
+	const FModifierKeysState ModifierKeys = Modifiers.ToModifierKeysState();
+	TSet<FKey> PressedButtons;
+	PressedButtons.Add(EKeys::LeftMouseButton);
+
+	// Move cursor to start, update hover tracking, then press.
+	FSlateApplication::Get().SetCursorPos(StartPosition);
+	FPointerEvent StartMoveEvent(
+		0, StartPosition, StartPosition,
+		TSet<FKey>(), EKeys::Invalid, 0.0f, ModifierKeys);
+	FSlateApplication::Get().ProcessMouseMoveEvent(StartMoveEvent);
+
+	FPointerEvent MouseDownEvent(
+		0, StartPosition, StartPosition,
+		PressedButtons, EKeys::LeftMouseButton, 0.0f, ModifierKeys);
+	FSlateApplication::Get().ProcessMouseButtonDownEvent(StartWindow->GetNativeWindow(), MouseDownEvent);
+
+	// Move cursor to end.
+	FPointerEvent MouseMoveEvent(
+		0, EndPosition, StartPosition,
+		PressedButtons, EKeys::Invalid, 0.0f, ModifierKeys);
+	FSlateApplication::Get().ProcessMouseMoveEvent(MouseMoveEvent);
+	FSlateApplication::Get().SetCursorPos(EndPosition);
+
+	// Release at end position.
+	FPointerEvent MouseUpEvent(
+		0, EndPosition, EndPosition,
+		TSet<FKey>(), EKeys::LeftMouseButton, 0.0f, ModifierKeys);
+	FSlateApplication::Get().ProcessMouseButtonUpEvent(MouseUpEvent);
+
+	return true;
+}
+
+// ── Window tools ─────────────────────────────────────────────────────
+
+FString USlateInspectorToolset::Windows(const FString& Action, int32 Index)
+{
+	if (!FSlateApplication::IsInitialized())
+	{
+		return TEXT("Error: Slate not initialized");
+	}
+
+	TArray<TSharedRef<SWindow>> VisibleWindows;
+	FSlateApplication::Get().GetAllVisibleWindowsOrdered(VisibleWindows);
+
+	if (Action.Equals(TEXT("list"), ESearchCase::IgnoreCase))
+	{
+		FString Result = TEXT("[");
+		for (int32 WindowIndex = 0; WindowIndex < VisibleWindows.Num(); ++WindowIndex)
+		{
+			const TSharedRef<SWindow>& Window = VisibleWindows[WindowIndex];
+			if (WindowIndex > 0)
+			{
+				Result += TEXT(", ");
+			}
+			FString Title = SlateInspectorToolsetPrivate::EscapeJsonString(Window->GetTitle().ToString());
+			Result += FString::Printf(TEXT("{\"index\": %d, \"title\": \"%s\"}"),
+				WindowIndex,
+				*Title);
+		}
+		Result += TEXT("]");
+		return Result;
+	}
+
+	if (Action.Equals(TEXT("select"), ESearchCase::IgnoreCase))
+	{
+		if (Index >= 0 && Index < VisibleWindows.Num())
+		{
+			VisibleWindows[Index]->BringToFront();
+			FSlateApplication::Get().SetAllUserFocus(VisibleWindows[Index], EFocusCause::SetDirectly);
+			return TEXT("OK");
+		}
+		return FString::Printf(TEXT("Error: Window index %d out of range (0..%d)"), Index, VisibleWindows.Num() - 1);
+	}
+
+	if (Action.Equals(TEXT("close"), ESearchCase::IgnoreCase))
+	{
+		if (Index >= 0 && Index < VisibleWindows.Num())
+		{
+			VisibleWindows[Index]->RequestDestroyWindow();
+			return TEXT("OK");
+		}
+		return FString::Printf(TEXT("Error: Window index %d out of range (0..%d)"), Index, VisibleWindows.Num() - 1);
+	}
+
+	return TEXT("Error: Unknown action. Use 'list', 'select', or 'close'.");
+}
+
+// ── Wait tool ────────────────────────────────────────────────────────
+
+bool USlateInspectorToolset::WaitFor(const FString& Text, const FString& TextGone)
+{
+	if (!FSlateApplication::IsInitialized())
+	{
+		return false;
+	}
+
+	const bool bWaitForAppear = !Text.IsEmpty();
+	const bool bWaitForDisappear = !TextGone.IsEmpty();
+
+	if (!bWaitForAppear && !bWaitForDisappear)
+	{
+		return true;
+	}
+
+	// Walk the widget tree looking for text without rendering a full snapshot,
+	// which would be unnecessary overhead for a simple text presence check.
+	bool bFoundText = false;
+	bool bFoundTextGone = false;
+
+	TArray<TSharedRef<SWindow>> Windows;
+	FSlateApplication::Get().GetAllVisibleWindowsOrdered(Windows);
+
+	TFunction<void(TSharedRef<SWidget>)> Walk = [&](TSharedRef<SWidget> Widget)
+	{
+		// Early out if both conditions already determined.
+		if ((!bWaitForAppear || bFoundText) && (!bWaitForDisappear || bFoundTextGone))
+		{
+			return;
+		}
+
+		FString WidgetText;
+		static const FName TextBlockType(TEXT("STextBlock"));
+		if (Widget->GetType() == TextBlockType)
+		{
+			WidgetText = StaticCastSharedRef<STextBlock>(Widget)->GetText().ToString();
+		}
+#if WITH_ACCESSIBILITY
+		if (WidgetText.IsEmpty())
+		{
+			WidgetText = Widget->GetAccessibleText(EAccessibleType::Main).ToString();
+		}
+#endif
+
+		if (!WidgetText.IsEmpty())
+		{
+			if (bWaitForAppear && WidgetText.Contains(Text))
+			{
+				bFoundText = true;
+			}
+			if (bWaitForDisappear && WidgetText.Contains(TextGone))
+			{
+				bFoundTextGone = true;
+			}
+		}
+
+		FChildren* Children = Widget->GetChildren();
+		for (int32 Index = 0; Index < Children->Num(); ++Index)
+		{
+			Walk(Children->GetChildAt(Index));
+		}
+	};
+
+	for (const TSharedRef<SWindow>& Window : Windows)
+	{
+		Walk(Window);
+	}
+
+	bool bAppearConditionMet = !bWaitForAppear || bFoundText;
+	bool bDisappearConditionMet = !bWaitForDisappear || !bFoundTextGone;
+
+	return bAppearConditionMet && bDisappearConditionMet;
+}
+
+// ── Form tool ────────────────────────────────────────────────────────
+
+bool USlateInspectorToolset::FillForm(const TArray<FSlateInspectorToolsetFormField>& Fields)
+{
+	bool bAllSuccess = true;
+
+	for (const FSlateInspectorToolsetFormField& Field : Fields)
+	{
+		bool bFieldSuccess = false;
+
+		if (Field.FieldType.Equals(TEXT("textbox"), ESearchCase::IgnoreCase))
+		{
+			bFieldSuccess = Type(Field.Ref, Field.Value, false);
+		}
+		else if (Field.FieldType.Equals(TEXT("checkbox"), ESearchCase::IgnoreCase))
+		{
+			// Click the checkbox if the desired state differs from the current state.
+			TSharedPtr<SWidget> Widget = ResolveRefOrWarn(Field.Ref);
+			static const FName CheckBoxType(TEXT("SCheckBox"));
+			if (Widget.IsValid() && Widget->GetType() == CheckBoxType)
+			{
+				TSharedRef<SCheckBox> CheckBox = StaticCastSharedRef<SCheckBox>(Widget.ToSharedRef());
+				bool bDesiredChecked = Field.Value.Equals(TEXT("true"), ESearchCase::IgnoreCase);
+				bool bCurrentlyChecked = CheckBox->IsChecked();
+				if (bDesiredChecked != bCurrentlyChecked)
+				{
+					bFieldSuccess = Click(Field.Ref);
+				}
+				else
+				{
+					bFieldSuccess = true; // Already in desired state.
+				}
+			}
+			else if (Widget.IsValid())
+			{
+				UE_LOGFMT(LogSlateInspectorToolset, Warning, "SlateInspectorToolset: Widget '{Ref}' is not an SCheckBox", Field.Ref);
+			}
+		}
+		else if (Field.FieldType.Equals(TEXT("combobox"), ESearchCase::IgnoreCase))
+		{
+			bFieldSuccess = SelectOption(Field.Ref, Field.Value);
+		}
+		else
+		{
+			UE_LOGFMT(LogSlateInspectorToolset, Warning, "SlateInspectorToolset: Unknown field type '{FieldType}'", Field.FieldType);
+		}
+
+		if (!bFieldSuccess)
+		{
+			bAllSuccess = false;
+		}
+	}
+
+	return bAllSuccess;
+}
diff --git a/Engine/Plugins/Experimental/SlateInspectorToolset/Source/SlateInspectorToolset/Private/SlateInspectorToolsetModule.cpp b/Engine/Plugins/Experimental/SlateInspectorToolset/Source/SlateInspectorToolset/Private/SlateInspectorToolsetModule.cpp
new file mode 100644
--- /dev/null
+++ b/Engine/Plugins/Experimental/SlateInspectorToolset/Source/SlateInspectorToolset/Private/SlateInspectorToolsetModule.cpp
@@ -0,0 +1,15 @@
+// Copyright Epic Games, Inc. All Rights Reserved.
+
+#include "SlateInspectorToolsetModule.h"
+
+#include "Modules/ModuleManager.h"
+
+IMPLEMENT_MODULE(FSlateInspectorToolsetModule, SlateInspectorToolset)
+
+void FSlateInspectorToolsetModule::StartupModule()
+{
+}
+
+void FSlateInspectorToolsetModule::ShutdownModule()
+{
+}
diff --git a/Engine/Plugins/Experimental/SlateInspectorToolset/Source/SlateInspectorToolset/Private/SlateInspectorToolsetObserverManager.cpp b/Engine/Plugins/Experimental/SlateInspectorToolset/Source/SlateInspectorToolset/Private/SlateInspectorToolsetObserverManager.cpp
new file mode 100644
--- /dev/null
+++ b/Engine/Plugins/Experimental/SlateInspectorToolset/Source/SlateInspectorToolset/Private/SlateInspectorToolsetObserverManager.cpp
@@ -0,0 +1,152 @@
+// Copyright Epic Games, Inc. All Rights Reserved.
+
+#include "SlateInspectorToolsetObserverManager.h"
+
+#include "SlateInspectorToolsetRefCache.h"
+#include "SlateInspectorToolsetSnapshotRenderer.h"
+
+FSlateInspectorToolsetObserverManager& FSlateInspectorToolsetObserverManager::Get()
+{
+	static FSlateInspectorToolsetObserverManager Instance;
+	return Instance;
+}
+
+FString FSlateInspectorToolsetObserverManager::AddObserver(TSharedPtr<SWidget> RootWidget, int32 MaxDepth)
+{
+	MaxDepth = FMath::Max(MaxDepth, 0);
+	bool bRoot = !RootWidget.IsValid();
+
+	// If an observer already exists for the same root, widen its depth (never shrink) and return its identifier.
+	for (FSlateInspectorToolsetObserver& Existing : Observers)
+	{
+		if (bRoot && Existing.bRoot)
+		{
+			Existing.MaxDepth = FMath::Max(Existing.MaxDepth, MaxDepth);
+			return Existing.Identifier;
+		}
+		if (!bRoot && !Existing.bRoot && Existing.RootWidget.Pin() == RootWidget)
+		{
+			Existing.MaxDepth = FMath::Max(Existing.MaxDepth, MaxDepth);
+			return Existing.Identifier;
+		}
+	}
+
+	FSlateInspectorToolsetObserver Observer;
+	Observer.Identifier = FString::Printf(TEXT("observer_%d"), NextIdentifier++);
+	Observer.bRoot = bRoot;
+	Observer.RootWidget = RootWidget;
+	Observer.MaxDepth = MaxDepth;
+
+	Observers.Add(MoveTemp(Observer));
+
+	return Observers.Last().Identifier;
+}
+
+bool FSlateInspectorToolsetObserverManager::RemoveObserver(const FString& Identifier)
+{
+	int32 Index = Observers.IndexOfByPredicate([&Identifier](const FSlateInspectorToolsetObserver& Observer)
+	{
+		return Observer.Identifier == Identifier;
+	});
+
+	if (Index == INDEX_NONE)
+	{
+		return false;
+	}
+
+	Observers.RemoveAt(Index);
+	return true;
+}
+
+TArray<FSlateInspectorToolsetObserver> FSlateInspectorToolsetObserverManager::GetObservers() const
+{
+	return Observers;
+}
+
+FString FSlateInspectorToolsetObserverManager::GetCachedSnapshot(const FString& Identifier) const
+{
+	for (const FSlateInspectorToolsetObserver& Observer : Observers)
+	{
+		if (Observer.Identifier == Identifier)
+		{
+			return Observer.CachedSnapshotText;
+		}
+	}
+
+	return FString();
+}
+
+FString FSlateInspectorToolsetObserverManager::FindMatchingObserverSnapshot(TSharedPtr<SWidget> RootWidget, int32 MaxDepth) const
+{
+	const FSlateInspectorToolsetObserver* BestMatch = nullptr;
+
+	for (const FSlateInspectorToolsetObserver& Observer : Observers)
+	{
+		if (Observer.MaxDepth < MaxDepth)
+		{
+			continue;
+		}
+
+		if (!RootWidget.IsValid() && Observer.bRoot)
+		{
+			if (!BestMatch || Observer.MaxDepth > BestMatch->MaxDepth)
+			{
+				BestMatch = &Observer;
+			}
+		}
+		else if (RootWidget.IsValid() && !Observer.bRoot)
+		{
+			TSharedPtr<SWidget> ObserverRoot = Observer.RootWidget.Pin();
+			if (ObserverRoot == RootWidget)
+			{
+				if (!BestMatch || Observer.MaxDepth > BestMatch->MaxDepth)
+				{
+					BestMatch = &Observer;
+				}
+			}
+		}
+	}
+
+	return BestMatch ? BestMatch->CachedSnapshotText : FString();
+}
+
+void FSlateInspectorToolsetObserverManager::Tick(float DeltaTime)
+{
+	AccumulatedTime += DeltaTime;
+
+	if (AccumulatedTime < TickIntervalSeconds)
+	{
+		return;
+	}
+	AccumulatedTime = 0.0f;
+
+	if (Observers.IsEmpty())
+	{
+		return;
+	}
+
+	// Remove non-root observers whose root widget has been destroyed.
+	Observers.RemoveAll([](const FSlateInspectorToolsetObserver& Observer)
+	{
+		return !Observer.bRoot && !Observer.RootWidget.IsValid();
+	});
+
+	// Walk each observer's subtree and cache the snapshot text.
+	for (FSlateInspectorToolsetObserver& Observer : Observers)
+	{
+		TSharedPtr<SWidget> Root = Observer.bRoot ? nullptr : Observer.RootWidget.Pin();
+		Observer.CachedSnapshotText = FSlateInspectorToolsetSnapshotRenderer::Render(
+			Root,
+			Observer.MaxDepth,
+			/* bIncludeSourceLocations */ false,
+			/* bResetCache */ false);
+	}
+
+	// Periodically purge expired refs.
+	++TickCounter;
+	if (TickCounter >= PurgeEveryNTicks)
+	{
+		TickCounter = 0;
+		FSlateInspectorToolsetRefCache::Get().PurgeExpired();
+	}
+}
diff --git a/Engine/Plugins/Experimental/SlateInspectorToolset/Source/SlateInspectorToolset/Private/SlateInspectorToolsetRefCache.cpp b/Engine/Plugins/Experimental/SlateInspectorToolset/Source/SlateInspectorToolset/Private/SlateInspectorToolsetRefCache.cpp
new file mode 100644
--- /dev/null
+++ b/Engine/Plugins/Experimental/SlateInspectorToolset/Source/SlateInspectorToolset/Private/SlateInspectorToolsetRefCache.cpp
@@ -0,0 +1,76 @@
+// Copyright Epic Games, Inc. All Rights Reserved.
+
+#include "SlateInspectorToolsetRefCache.h"
+
+FSlateInspectorToolsetRefCache& FSlateInspectorToolsetRefCache::Get()
+{
+	static FSlateInspectorToolsetRefCache Instance;
+	return Instance;
+}
+
+FString FSlateInspectorToolsetRefCache::GetOrAssignRef(TSharedRef<SWidget> Widget, const FString& RolePrefix)
+{
+	TWeakPtr<SWidget> WeakWidget = Widget;
+
+	if (const FString* ExistingRef = WidgetToRef.Find(WeakWidget))
+	{
+		return *ExistingRef;
+	}
+
+	int32& Counter = RolePrefixCounters.FindOrAdd(RolePrefix, 0);
+	++Counter;
+
+	FString Ref = FString::Printf(TEXT("%s%d"), *RolePrefix, Counter);
+
+	WidgetToRef.Add(WeakWidget, Ref);
+	RefToWidget.Add(Ref, WeakWidget);
+
+	return Ref;
+}
+
+FString FSlateInspectorToolsetRefCache::FindRef(TSharedRef<SWidget> Widget) const
+{
+	TWeakPtr<SWidget> WeakWidget = Widget;
+	if (const FString* Ref = WidgetToRef.Find(WeakWidget))
+	{
+		return *Ref;
+	}
+	return FString();
+}
+
+TSharedPtr<SWidget> FSlateInspectorToolsetRefCache::ResolveRef(const FString& Ref) const
+{
+	if (const TWeakPtr<SWidget>* WeakWidget = RefToWidget.Find(Ref))
+	{
+		return WeakWidget->Pin();
+	}
+	return nullptr;
+}
+
+void FSlateInspectorToolsetRefCache::Reset()
+{
+	WidgetToRef.Empty();
+	RefToWidget.Empty();
+	RolePrefixCounters.Empty();
+}
+
+void FSlateInspectorToolsetRefCache::PurgeExpired()
+{
+	// Capture both the key and its ref string in a single pass.  Expired
+	// weak pointers may hash identically, so a second Find() on the key
+	// alone could silently fail or match the wrong entry.
+	TArray<TPair<TWeakPtr<SWidget>, FString>> Expired;
+	for (const auto& Pair : WidgetToRef)
+	{
+		if (!Pair.Key.IsValid())
+		{
+			Expired.Emplace(Pair.Key, Pair.Value);
+		}
+	}
+
+	for (const auto& Entry : Expired)
+	{
+		RefToWidget.Remove(Entry.Value);
+		WidgetToRef.Remove(Entry.Key);
+	}
+}
diff --git a/Engine/Plugins/Experimental/SlateInspectorToolset/Source/SlateInspectorToolset/Private/SlateInspectorToolsetSnapshotRenderer.cpp b/Engine/Plugins/Experimental/SlateInspectorToolset/Source/SlateInspectorToolset/Private/SlateInspectorToolsetSnapshotRenderer.cpp
new file mode 100644
--- /dev/null
+++ b/Engine/Plugins/Experimental/SlateInspectorToolset/Source/SlateInspectorToolset/Private/SlateInspectorToolsetSnapshotRenderer.cpp
@@ -0,0 +1,414 @@
+// Copyright Epic Games, Inc. All Rights Reserved.
+
+#include "SlateInspectorToolsetSnapshotRenderer.h"
+
+#include "SlateInspectorToolsetRefCache.h"
+
+#include "Framework/Application/SlateApplication.h"
+#include "Misc/Paths.h"
+#include "Widgets/Input/SCheckBox.h"
+#include "Widgets/Input/SEditableTextBox.h"
+#include "Widgets/Docking/SDockTab.h"
+#include "Widgets/SWindow.h"
+#include "Widgets/Text/STextBlock.h"
+
+TMap<FName, FString> FSlateInspectorToolsetSnapshotRenderer::TypeToRole;
+TArray<TPair<FString, FString>> FSlateInspectorToolsetSnapshotRenderer::TemplatePrefixToRole;
+TMap<FString, FString> FSlateInspectorToolsetSnapshotRenderer::RoleToRefPrefix;
+TSet<FName> FSlateInspectorToolsetSnapshotRenderer::StructuralContainers;
+TMap<FName, TFunction<FString(TSharedRef<SWidget>)>> FSlateInspectorToolsetSnapshotRenderer::CustomLabelExtractors;
+TMap<FName, TFunction<TArray<FString>(TSharedRef<SWidget>)>> FSlateInspectorToolsetSnapshotRenderer::CustomStateFlagsExtractors;
+bool FSlateInspectorToolsetSnapshotRenderer::bMapsInitialized = false;
+
+void FSlateInspectorToolsetSnapshotRenderer::EnsureMapsInitialized()
+{
+	if (bMapsInitialized)
+	{
+		return;
+	}
+	bMapsInitialized = true;
+
+	// Type -> role mapping (non-template types use exact FName match).
+	TypeToRole.Add(FName(TEXT("SWindow")),               TEXT("window"));
+	TypeToRole.Add(FName(TEXT("SButton")),               TEXT("button"));
+	TypeToRole.Add(FName(TEXT("SEditableText")),          TEXT("textbox"));
+	TypeToRole.Add(FName(TEXT("SEditableTextBox")),       TEXT("textbox"));
+	TypeToRole.Add(FName(TEXT("SMultiLineEditableText")), TEXT("textbox"));
+	TypeToRole.Add(FName(TEXT("SCheckBox")),              TEXT("checkbox"));
+	TypeToRole.Add(FName(TEXT("SComboButton")),           TEXT("combobox"));
+	TypeToRole.Add(FName(TEXT("SSlider")),                TEXT("slider"));
+	TypeToRole.Add(FName(TEXT("SMenuBar")),               TEXT("menubar"));
+	TypeToRole.Add(FName(TEXT("SMenuAnchor")),            TEXT("menu"));
+	TypeToRole.Add(FName(TEXT("SScrollBox")),             TEXT("scrollable"));
+	TypeToRole.Add(FName(TEXT("SSplitter")),              TEXT("splitter"));
+	TypeToRole.Add(FName(TEXT("SDockTab")),               TEXT("tab"));
+	TypeToRole.Add(FName(TEXT("STabWell")),               TEXT("tablist"));
+	TypeToRole.Add(FName(TEXT("SToolTip")),               TEXT("tooltip"));
+	TypeToRole.Add(FName(TEXT("STextBlock")),             TEXT("text"));
+	TypeToRole.Add(FName(TEXT("SRichTextBlock")),         TEXT("text"));
+	TypeToRole.Add(FName(TEXT("SImage")),                 TEXT("image"));
+	TypeToRole.Add(FName(TEXT("SProgressBar")),           TEXT("progressbar"));
+
+	// Template types use prefix matching (e.g., "SComboBox<" matches "SComboBox<TSharedPtr<FString>>").
+	TemplatePrefixToRole.Add(TPair<FString, FString>(TEXT("SComboBox<"),  TEXT("combobox")));
+	TemplatePrefixToRole.Add(TPair<FString, FString>(TEXT("SSpinBox<"),   TEXT("slider")));
+	TemplatePrefixToRole.Add(TPair<FString, FString>(TEXT("STreeView<"),  TEXT("treeview")));
+	TemplatePrefixToRole.Add(TPair<FString, FString>(TEXT("SListView<"),  TEXT("listview")));
+	TemplatePrefixToRole.Add(TPair<FString, FString>(TEXT("STableRow<"),  TEXT("listitem")));
+
+	// Role -> ref prefix mapping.
+	RoleToRefPrefix.Add(TEXT("window"),      TEXT("w"));
+	RoleToRefPrefix.Add(TEXT("button"),      TEXT("b"));
+	RoleToRefPrefix.Add(TEXT("textbox"),     TEXT("tb"));
+	RoleToRefPrefix.Add(TEXT("checkbox"),    TEXT("cb"));
+	RoleToRefPrefix.Add(TEXT("combobox"),    TEXT("co"));
+	RoleToRefPrefix.Add(TEXT("slider"),      TEXT("sl"));
+	RoleToRefPrefix.Add(TEXT("menubar"),     TEXT("mb"));
+	RoleToRefPrefix.Add(TEXT("menu"),        TEXT("m"));
+	RoleToRefPrefix.Add(TEXT("treeview"),    TEXT("tv"));
+	RoleToRefPrefix.Add(TEXT("listview"),    TEXT("lv"));
+	RoleToRefPrefix.Add(TEXT("listitem"),    TEXT("li"));
+	RoleToRefPrefix.Add(TEXT("scrollable"),  TEXT("sc"));
+	RoleToRefPrefix.Add(TEXT("splitter"),    TEXT("sp"));
+	RoleToRefPrefix.Add(TEXT("tab"),         TEXT("t"));
+	RoleToRefPrefix.Add(TEXT("tablist"),     TEXT("tl"));
+	RoleToRefPrefix.Add(TEXT("tooltip"),     TEXT("tt"));
+	RoleToRefPrefix.Add(TEXT("text"),        TEXT("x"));
+	RoleToRefPrefix.Add(TEXT("image"),       TEXT("i"));
+	RoleToRefPrefix.Add(TEXT("progressbar"), TEXT("pb"));
+
+	// Structural containers to skip (promote children).
+	StructuralContainers.Add(FName(TEXT("SBox")));
+	StructuralContainers.Add(FName(TEXT("SBorder")));
+	StructuralContainers.Add(FName(TEXT("SOverlay")));
+	StructuralContainers.Add(FName(TEXT("SSpacer")));
+	StructuralContainers.Add(FName(TEXT("SConstraintCanvas")));
+	StructuralContainers.Add(FName(TEXT("SHorizontalBox")));
+	StructuralContainers.Add(FName(TEXT("SVerticalBox")));
+	StructuralContainers.Add(FName(TEXT("SGridPanel")));
+	StructuralContainers.Add(FName(TEXT("SWrapBox")));
+	StructuralContainers.Add(FName(TEXT("SWidgetSwitcher")));
+	StructuralContainers.Add(FName(TEXT("SCanvas")));
+	StructuralContainers.Add(FName(TEXT("SScaleBox")));
+	StructuralContainers.Add(FName(TEXT("SSizeBox")));
+	StructuralContainers.Add(FName(TEXT("SNullWidget")));
+	StructuralContainers.Add(FName(TEXT("SInvalidationPanel")));
+	StructuralContainers.Add(FName(TEXT("SRetainerWidget")));
+}
+
+void FSlateInspectorToolsetSnapshotRenderer::RegisterWidgetRole(const FName& WidgetType, const FString& Role, const FString& RefPrefix)
+{
+	EnsureMapsInitialized();
+	TypeToRole.Add(WidgetType, Role);
+	RoleToRefPrefix.Add(Role, RefPrefix);
+}
+
+void FSlateInspectorToolsetSnapshotRenderer::RegisterLabelExtractor(const FName& WidgetType, TFunction<FString(TSharedRef<SWidget>)> Extractor)
+{
+	EnsureMapsInitialized();
+	CustomLabelExtractors.Add(WidgetType, MoveTemp(Extractor));
+}
+
+void FSlateInspectorToolsetSnapshotRenderer::RegisterStateFlagsExtractor(const FName& WidgetType, TFunction<TArray<FString>(TSharedRef<SWidget>)> Extractor)
+{
+	EnsureMapsInitialized();
+	CustomStateFlagsExtractors.Add(WidgetType, MoveTemp(Extractor));
+}
+
+const FString* FSlateInspectorToolsetSnapshotRenderer::FindRoleForType(const FName& WidgetType, const FString& TypeString)
+{
+	// Try exact match first.
+	if (const FString* Role = TypeToRole.Find(WidgetType))
+	{
+		return Role;
+	}
+
+	// Try prefix match for template types (e.g., "SComboBox<TSharedPtr<FString>>").
+	for (const TPair<FString, FString>& Pair : TemplatePrefixToRole)
+	{
+		if (TypeString.StartsWith(Pair.Key))
+		{
+			return &Pair.Value;
+		}
+	}
+
+	return nullptr;
+}
+
+FString FSlateInspectorToolsetSnapshotRenderer::Render(TSharedPtr<SWidget> Root, int32 MaxDepth, bool bIncludeSourceLocations, bool bResetCache)
+{
+	EnsureMapsInitialized();
+
+	MaxDepth = FMath::Max(MaxDepth, 0);
+
+	if (bResetCache)
+	{
+		FSlateInspectorToolsetRefCache::Get().Reset();
+	}
+
+	TStringBuilder<4096> Output;
+
+	if (Root.IsValid())
+	{
+		RenderWidget(Output, Root.ToSharedRef(), 0, 0, MaxDepth, bIncludeSourceLocations);
+	}
+	else
+	{
+		// Snapshot all visible windows in z-order.
+		if (FSlateApplication::IsInitialized())
+		{
+			TArray<TSharedRef<SWindow>> VisibleWindows;
+			FSlateApplication::Get().GetAllVisibleWindowsOrdered(VisibleWindows);
+			for (const TSharedRef<SWindow>& Window : VisibleWindows)
+			{
+				RenderWidget(Output, Window, 0, 0, MaxDepth, bIncludeSourceLocations);
+			}
+		}
+	}
+
+	return Output.ToString();
+}
+
+void FSlateInspectorToolsetSnapshotRenderer::RenderWidget(
+	FStringBuilderBase& Output,
+	TSharedRef<SWidget> Widget,
+	int32 IndentLevel,
+	int32 CurrentDepth,
+	int32 MaxDepth,
+	bool bIncludeSourceLocations)
+{
+	if (CurrentDepth > MaxDepth)
+	{
+		return;
+	}
+
+	// Skip collapsed/hidden widgets.
+	if (Widget->GetVisibility() == EVisibility::Collapsed || Widget->GetVisibility() == EVisibility::Hidden)
+	{
+		return;
+	}
+
+	const FName WidgetType = Widget->GetType();
+	const FString TypeString = Widget->GetTypeAsString();
+
+	// Check if this is a structural container (skip but recurse children at same indent).
+	if (StructuralContainers.Contains(WidgetType))
+	{
+		FChildren* Children = Widget->GetChildren();
+		for (int32 ChildIndex = 0; ChildIndex < Children->Num(); ++ChildIndex)
+		{
+			RenderWidget(Output, Children->GetChildAt(ChildIndex), IndentLevel, CurrentDepth + 1, MaxDepth, bIncludeSourceLocations);
+		}
+		return;
+	}
+
+	// Look up role for this widget type (handles both exact and template prefix match).
+	const FString* Role = FindRoleForType(WidgetType, TypeString);
+
+	// Text collapsing: if this widget has exactly one child and that child is STextBlock,
+	// absorb the text child's label and don't recurse into it.
+	bool bAbsorbedTextChild = false;
+	FString AbsorbedText;
+	FChildren* Children = Widget->GetChildren();
+	if (Children->Num() == 1)
+	{
+		TSharedRef<SWidget> OnlyChild = Children->GetChildAt(0);
+		static const FName TextBlockType(TEXT("STextBlock"));
+		if (OnlyChild->GetType() == TextBlockType)
+		{
+			TSharedRef<STextBlock> TextBlock = StaticCastSharedRef<STextBlock>(OnlyChild);
+			AbsorbedText = TextBlock->GetText().ToString();
+			bAbsorbedTextChild = true;
+		}
+	}
+
+	// If we have no role and no useful label, skip this widget but recurse children.
+	if (!Role)
+	{
+		FString Label = ExtractLabel(Widget, TypeString);
+		if (Label.IsEmpty() && !bAbsorbedTextChild)
+		{
+			for (int32 ChildIndex = 0; ChildIndex < Children->Num(); ++ChildIndex)
+			{
+				RenderWidget(Output, Children->GetChildAt(ChildIndex), IndentLevel, CurrentDepth + 1, MaxDepth, bIncludeSourceLocations);
+			}
+			return;
+		}
+	}
+
+	// This widget is meaningful, emit it.
+	FString Label;
+	if (bAbsorbedTextChild && !AbsorbedText.IsEmpty())
+	{
+		Label = AbsorbedText;
+	}
+	else
+	{
+		Label = ExtractLabel(Widget, TypeString);
+	}
+
+	FString StateFlags = ExtractStateFlags(Widget, TypeString);
+
+	// Determine the role string and ref prefix.
+	FString RoleString = Role ? *Role : TEXT("generic");
+	const FString* RefPrefix = RoleToRefPrefix.Find(RoleString);
+	FString Prefix = RefPrefix ? *RefPrefix : TEXT("g");
+
+	FString Ref = FSlateInspectorToolsetRefCache::Get().GetOrAssignRef(Widget, Prefix);
+
+	// Emit the line.
+	for (int32 Index = 0; Index < IndentLevel; ++Index)
+	{
+		Output.Append(TEXT("  "));
+	}
+
+	Output.Append(*RoleString);
+
+	if (!Label.IsEmpty())
+	{
+		Output.Append(TEXT(" \""));
+		// Escape quotes and control characters that would break the line-oriented format.
+		for (TCHAR Character : Label)
+		{
+			switch (Character)
+			{
+			case TEXT('"'):  Output.Append(TEXT("\\\"")); break;
+			case TEXT('\\'): Output.Append(TEXT("\\\\")); break;
+			case TEXT('\n'): Output.Append(TEXT("\\n"));  break;
+			case TEXT('\r'): Output.Append(TEXT("\\r"));  break;
+			case TEXT('\t'): Output.Append(TEXT("\\t"));  break;
+			default:         Output.AppendChar(Character); break;
+			}
+		}
+		Output.Append(TEXT("\""));
+	}
+
+	if (!StateFlags.IsEmpty())
+	{
+		Output.Append(TEXT(" "));
+		Output.Append(*StateFlags);
+	}
+
+	// Geometry: absolute position (desktop space) and local size (widget space).
+	const FGeometry& Geometry = Widget->GetCachedGeometry();
+	const FVector2D AbsolutePosition = Geometry.GetAbsolutePosition();
+	const FVector2D LocalSize = Geometry.GetLocalSize();
+	Output.Appendf(TEXT(" [pos=%.0f,%.0f size=%.0f,%.0f]"), AbsolutePosition.X, AbsolutePosition.Y, LocalSize.X, LocalSize.Y);
+
+	if (bIncludeSourceLocations)
+	{
+#if !UE_BUILD_SHIPPING
+		FName CreatedInLocation = Widget->GetCreatedInLocation();
+		if (!CreatedInLocation.IsNone())
+		{
+			FString FilePath = CreatedInLocation.GetPlainNameString();
+			int32 LineNumber = CreatedInLocation.GetNumber();
+			FString FileName = FPaths::GetCleanFilename(FilePath);
+			if (LineNumber > 0)
+			{
+				Output.Appendf(TEXT(" [src=%s:%d]"), *FileName, LineNumber);
+			}
+			else
+			{
+				Output.Appendf(TEXT(" [src=%s]"), *FileName);
+			}
+		}
+#endif
+	}
+
+	Output.Append(TEXT(" [ref="));
+	Output.Append(*Ref);
+	Output.Append(TEXT("]\n"));
+
+	// Recurse into children (unless we absorbed the only text child).
+	if (!bAbsorbedTextChild)
+	{
+		for (int32 ChildIndex = 0; ChildIndex < Children->Num(); ++ChildIndex)
+		{
+			RenderWidget(Output, Children->GetChildAt(ChildIndex), IndentLevel + 1, CurrentDepth + 1, MaxDepth, bIncludeSourceLocations);
+		}
+	}
+}
+
+FString FSlateInspectorToolsetSnapshotRenderer::ExtractLabel(TSharedRef<SWidget> Widget, const FString& TypeString)
+{
+	// Prefer Slate's native accessible text since many widgets implement this already.
+#if WITH_ACCESSIBILITY
+	FText AccessibleText = Widget->GetAccessibleText(EAccessibleType::Main);
+	if (!AccessibleText.IsEmpty())
+	{
+		return AccessibleText.ToString();
+	}
+#endif
+
+	// Fall back to type-specific extraction when no accessible text is available.
+	if (TypeString == TEXT("SWindow"))
+	{
+		return StaticCastSharedRef<SWindow>(Widget)->GetTitle().ToString();
+	}
+	if (TypeString == TEXT("STextBlock"))
+	{
+		return StaticCastSharedRef<STextBlock>(Widget)->GetText().ToString();
+	}
+	if (TypeString == TEXT("SEditableTextBox"))
+	{
+		return StaticCastSharedRef<SEditableTextBox>(Widget)->GetText().ToString();
+	}
+	if (TypeString == TEXT("SDockTab"))
+	{
+		return StaticCastSharedRef<SDockTab>(Widget)->GetTabLabel().ToString();
+	}
+
+	// Try custom label extractors.
+	if (const auto* Extractor = CustomLabelExtractors.Find(Widget->GetType()))
+	{
+		return (*Extractor)(Widget);
+	}
+
+	return FString();
+}
+
+FString FSlateInspectorToolsetSnapshotRenderer::ExtractStateFlags(TSharedRef<SWidget> Widget, const FString& TypeString)
+{
+	TArray<FString> Flags;
+
+	if (!Widget->IsEnabled())
+	{
+		Flags.Add(TEXT("disabled"));
+	}
+
+	if (Widget->HasKeyboardFocus())
+	{
+		Flags.Add(TEXT("focused"));
+	}
+
+	if (TypeString == TEXT("SCheckBox"))
+	{
+		ECheckBoxState CheckedState = StaticCastSharedRef<SCheckBox>(Widget)->GetCheckedState();
+		switch (CheckedState)
+		{
+		case ECheckBoxState::Checked:
+			Flags.Add(TEXT("checked"));
+			break;
+		case ECheckBoxState::Undetermined:
+			Flags.Add(TEXT("indeterminate"));
+			break;
+		default:
+			Flags.Add(TEXT("unchecked"));
+			break;
+		}
+	}
+
+	// Try custom state flags extractors.
+	if (const auto* Extractor = CustomStateFlagsExtractors.Find(Widget->GetType()))
+	{
+		Flags.Append((*Extractor)(Widget));
+	}
+
+	if (Flags.IsEmpty())
+	{
+		return FString();
+	}
+
+	return FString::Printf(TEXT("[%s]"), *FString::Join(Flags, TEXT(", ")));
+}
diff --git a/Engine/Plugins/Experimental/SlateInspectorToolset/Source/SlateInspectorToolset/Private/SlateInspectorToolsetSubsystem.cpp b/Engine/Plugins/Experimental/SlateInspectorToolset/Source/SlateInspectorToolset/Private/SlateInspectorToolsetSubsystem.cpp
new file mode 100644
--- /dev/null
+++ b/Engine/Plugins/Experimental/SlateInspectorToolset/Source/SlateInspectorToolset/Private/SlateInspectorToolsetSubsystem.cpp
@@ -0,0 +1,83 @@
+// Copyright Epic Games, Inc. All Rights Reserved.
+
+#include "SlateInspectorToolsetSubsystem.h"
+
+#include "Editor.h"
+#include "Framework/Application/SlateApplication.h"
+#include "SlateInspectorToolsetObserverManager.h"
+#include "SlateInspectorToolset.h"
+#include "Subsystems/SubsystemCollection.h"
+#include "ToolsetRegistry/UToolsetRegistry.h"
+
+static bool bEnableSlateInspectorToolset = true;
+
+static void OnEnableSlateInspectorToolsetChanged(IConsoleVariable* Variable)
+{
+	if (GEditor == nullptr)
+	{
+		return;
+	}
+	
+	if (USlateInspectorToolsetSubsystem* Subsystem = GEditor->GetEditorSubsystem<USlateInspectorToolsetSubsystem>())
+	{
+		Subsystem->SetToolsetEnabled(Variable->GetBool());
+	}
+}
+
+static FAutoConsoleVariableRef CVarEnableSlateInspectorToolset(
+	TEXT("SlateInspectorToolset.Enable"),
+	bEnableSlateInspectorToolset,
+	TEXT("Enable or disable SlateInspectorToolset registration. When disabled, its tools will not appear in the MCP tool list."),
+	FConsoleVariableDelegate::CreateStatic(&OnEnableSlateInspectorToolsetChanged),
+	ECVF_Default);
+
+void USlateInspectorToolsetSubsystem::Initialize(FSubsystemCollectionBase& Collection)
+{
+	Super::Initialize(Collection);
+	SetToolsetEnabled(bEnableSlateInspectorToolset);
+
+	if (FSlateApplication::IsInitialized())
+	{
+		PostTickDelegateHandle = FSlateApplication::Get().OnPostTick().AddUObject(this, &USlateInspectorToolsetSubsystem::OnSlatePostTick);
+	}
+
+	// Create a shallow default root observer (depth 0: window names only).
+	// Use Observe() on specific windows or panels to drill deeper.
+	RootObserverIdentifier = FSlateInspectorToolsetObserverManager::Get().AddObserver(nullptr, 0);
+}
+
+void USlateInspectorToolsetSubsystem::Deinitialize()
+{
+	if (!RootObserverIdentifier.IsEmpty())
+	{
+		FSlateInspectorToolsetObserverManager::Get().RemoveObserver(RootObserverIdentifier);
+		RootObserverIdentifier.Empty();
+	}
+
+	if (FSlateApplication::IsInitialized())
+	{
+		FSlateApplication::Get().OnPostTick().Remove(PostTickDelegateHandle);
+	}
+
+	SetToolsetEnabled(false);
+	Super::Deinitialize();
+}
+
+void USlateInspectorToolsetSubsystem::OnSlatePostTick(float DeltaTime)
+{
+	FSlateInspectorToolsetObserverManager::Get().Tick(DeltaTime);
+}
+
+void USlateInspectorToolsetSubsystem::SetToolsetEnabled(bool bEnabled)
+{
+	if (bEnabled && !bToolsetRegistered)
+	{
+		UToolsetRegistry::RegisterToolsetClass(USlateInspectorToolset::StaticClass());
+		bToolsetRegistered = true;
+	}
+	else if (!bEnabled && bToolsetRegistered)
+	{
+		UToolsetRegistry::UnregisterToolsetClass(USlateInspectorToolset::StaticClass());
+		bToolsetRegistered = false;
+	}
+}
diff --git a/Engine/Plugins/Experimental/SlateInspectorToolset/Source/SlateInspectorToolset/Public/SlateInspectorToolset.h b/Engine/Plugins/Experimental/SlateInspectorToolset/Source/SlateInspectorToolset/Public/SlateInspectorToolset.h
new file mode 100644
--- /dev/null
+++ b/Engine/Plugins/Experimental/SlateInspectorToolset/Source/SlateInspectorToolset/Public/SlateInspectorToolset.h
@@ -0,0 +1,196 @@
+// Copyright Epic Games, Inc. All Rights Reserved.
+
+#pragma once
+
+#include "Input/Events.h"
+#include "ToolsetRegistry/UToolsetRegistry.h"
+#include "ToolsetRegistry/ToolsetImage.h"
+#include "Widgets/SWidget.h"
+
+#include "SlateInspectorToolset.generated.h"
+
+/** Modifier keys held during an input simulation (click, drag, etc.). */
+USTRUCT(BlueprintType)
+struct FSlateInspectorToolsetModifierKeys
+{
+	GENERATED_BODY()
+
+	UPROPERTY(BlueprintReadWrite, Category = "SlateInspectorToolset")
+	bool bShift = false;
+
+	UPROPERTY(BlueprintReadWrite, Category = "SlateInspectorToolset")
+	bool bCtrl = false;
+
+	UPROPERTY(BlueprintReadWrite, Category = "SlateInspectorToolset")
+	bool bAlt = false;
+
+	UPROPERTY(BlueprintReadWrite, Category = "SlateInspectorToolset")
+	bool bCmd = false;
+
+	/** Convert to Slate's FModifierKeysState. */
+	FModifierKeysState ToModifierKeysState() const
+	{
+		return FModifierKeysState(bShift, bShift, bCtrl, bCtrl, bAlt, bAlt, bCmd, bCmd, false);
+	}
+};
+
+/** Describes a single form field for the FillForm tool. */
+USTRUCT(BlueprintType)
+struct FSlateInspectorToolsetFormField
+{
+	GENERATED_BODY()
+
+	/** The ref identifier of the form field widget. */
+	UPROPERTY(BlueprintReadWrite, Category = "SlateInspectorToolset")
+	FString Ref;
+
+	/** The value to set. */
+	UPROPERTY(BlueprintReadWrite, Category = "SlateInspectorToolset")
+	FString Value;
+
+	/** The type of field: "textbox", "checkbox", or "combobox". */
+	UPROPERTY(BlueprintReadWrite, Category = "SlateInspectorToolset")
+	FString FieldType;
+};
+
+/**
+ * Playwright-style Slate UI automation toolset.
+ *
+ * Exposes snapshot, screenshot, and interaction tools for driving the
+ * Unreal Editor UI programmatically.  Registered via UToolsetRegistry
+ * so the ModelContextProtocol plugin picks them up automatically.
+ *
+ * A shallow root observer (depth 0) continuously tracks top-level
+ * windows.  Before working with a specific window or panel, call
+ * Observe() on it to get deep widget coverage, then Unobserve() when
+ * done.  Observers walk their subtree every ~100ms, assigning refs to
+ * newly appeared widgets and keeping the ref cache current.
+ *
+ * Input simulation uses direct Slate event APIs (ProcessKeyCharEvent,
+ * ProcessMouseButtonDownEvent, etc.) rather than the AutomationDriver,
+ * because AutomationDriver's synchronous API deadlocks when called from
+ * the game thread (which is where MCP tool calls execute).
+ */
+UCLASS(BlueprintType, MinimalAPI)
+class USlateInspectorToolset : public UToolsetDefinition
+{
+	GENERATED_BODY()
+
+public:
+
+	/** Capture a Slate UI accessibility snapshot.  Use this to read the current
+	 * widget tree and discover refs for action tools (Click, Type, Hover, etc.).
+	 * A shallow root observer (depth 0) covers top-level windows automatically.
+	 * Before interacting with a specific window or panel, call Observe() on it
+	 * to get deep coverage, then Snapshot that subtree to see its contents.
+	 * Refs discovered by a previous Snapshot remain usable. You do NOT need to
+	 * call Snapshot again before every action.
+	 * @param Ref                      Subtree root ref. Empty = all windows.
+	 * @param MaxDepth                 Maximum depth (default 30).
+	 * @param bIncludeSourceLocations  Include [src=File:Line] tags showing where each widget was created in C++. */
+	UFUNCTION(BlueprintCallable, Category = "SlateInspectorToolset")
+	static SLATEINSPECTORTOOLSET_API FString Snapshot(const FString& Ref, int32 MaxDepth = 30, bool bIncludeSourceLocations = false);
+
+	/** Register an observer on a widget subtree so its refs are continuously
+	 * kept up to date (~100ms tick).  Call this on the window or panel you are
+	 * about to work with. It ensures new widgets appearing in that subtree
+	 * are assigned refs automatically.  Unobserve when you are done.
+	 * A shallow root observer (depth 0) already covers top-level windows.
+	 * @param Ref       Root widget ref to observe. Empty = all visible windows.
+	 * @param MaxDepth  Maximum depth to walk from the root. */
+	UFUNCTION(BlueprintCallable, Category = "SlateInspectorToolset")
+	static SLATEINSPECTORTOOLSET_API FString Observe(const FString& Ref, int32 MaxDepth = 30);
+
+	/** Remove an observer by its identifier.
+	 * @param Identifier  The identifier returned by Observe(). */
+	UFUNCTION(BlueprintCallable, Category = "SlateInspectorToolset")
+	static SLATEINSPECTORTOOLSET_API bool Unobserve(const FString& Identifier);
+
+	/** List all active observers as a JSON array for debugging.
+	 * Each entry includes the observer identifier, whether it is the root observer,
+	 * the root widget ref (if any), max depth, and cached snapshot size. */
+	UFUNCTION(BlueprintCallable, Category = "SlateInspectorToolset")
+	static SLATEINSPECTORTOOLSET_API FString ListObservers();
+
+	/** Screenshot a Slate widget or the active editor window. Prefer this over
+	 * SceneTools.take_screenshot for Editor UI; use SceneTools only for 3D viewport.
+	 * @param Ref  Slate widget ref. Empty = active window. */
+	UFUNCTION(BlueprintCallable, Category = "SlateInspectorToolset")
+	static SLATEINSPECTORTOOLSET_API FToolsetImage Screenshot(const FString& Ref);
+
+	// -- Action tools --
+
+	/** Click a Slate widget identified by its ref.
+	 * @param Ref         Slate widget ref.
+	 * @param Button      "left", "right", or "middle".
+	 * @param DoubleClick True for double-click.
+	 * @param Modifiers   Modifier keys held during the click. */
+	UFUNCTION(BlueprintCallable, Category = "SlateInspectorToolset")
+	static SLATEINSPECTORTOOLSET_API bool Click(const FString& Ref, const FString& Button = TEXT("left"), bool DoubleClick = false, const FSlateInspectorToolsetModifierKeys& Modifiers = FSlateInspectorToolsetModifierKeys());
+
+	/** Hover over a Slate widget, triggering any hover state or tooltip.
+	 * @param Ref  Slate widget ref. */
+	UFUNCTION(BlueprintCallable, Category = "SlateInspectorToolset")
+	static SLATEINSPECTORTOOLSET_API bool Hover(const FString& Ref);
+
+	/** Type text into a Slate text input widget. Focuses the widget first,
+	 * then sends one key event per character.
+	 * @param Ref    Slate textbox ref.
+	 * @param Text   Text to type.
+	 * @param Submit Press Enter after typing. */
+	UFUNCTION(BlueprintCallable, Category = "SlateInspectorToolset")
+	static SLATEINSPECTORTOOLSET_API bool Type(const FString& Ref, const FString& Text, bool Submit = false);
+
+	/** Press and release a keyboard key on the currently focused Slate widget.
+	 * Supports modifier prefixes: "Ctrl+C", "Shift+1".
+	 * @param Key  Key name with optional modifiers, e.g. "Enter", "Ctrl+A", "Shift+Tab". */
+	UFUNCTION(BlueprintCallable, Category = "SlateInspectorToolset")
+	static SLATEINSPECTORTOOLSET_API bool PressKey(const FString& Key);
+
+	/** Select an option in a Slate combobox by its text label.
+	 * Opens the dropdown, finds the matching text, and clicks it.
+	 * @param Ref    Slate combobox ref.
+	 * @param Value  Exact option text to select. */
+	UFUNCTION(BlueprintCallable, Category = "SlateInspectorToolset")
+	static SLATEINSPECTORTOOLSET_API bool SelectOption(const FString& Ref, const FString& Value);
+
+	/** Drag from one Slate widget to another (mouse down, move, release).
+	 * @param StartRef   Slate widget ref for the drag source.
+	 * @param EndRef     Slate widget ref for the drop target.
+	 * @param Modifiers  Modifier keys held during the drag. */
+	UFUNCTION(BlueprintCallable, Category = "SlateInspectorToolset")
+	static SLATEINSPECTORTOOLSET_API bool Drag(const FString& StartRef, const FString& EndRef, const FSlateInspectorToolsetModifierKeys& Modifiers = FSlateInspectorToolsetModifierKeys());
+
+	/** List, select, or close top-level Slate editor windows.
+	 * @param Action  "list" returns JSON array, "select" brings to front, "close" destroys.
+	 * @param Index   Window index for select/close. */
+	UFUNCTION(BlueprintCallable, Category = "SlateInspectorToolset")
+	static SLATEINSPECTORTOOLSET_API FString Windows(const FString& Action = TEXT("list"), int32 Index = -1);
+
+	/** Check if text is present or absent in the Slate widget tree.
+	 * Non-blocking: checks once and returns immediately. Poll to wait.
+	 * @param Text       Text that must be present (empty = skip).
+	 * @param TextGone   Text that must be absent (empty = skip). */
+	UFUNCTION(BlueprintCallable, Category = "SlateInspectorToolset")
+	static SLATEINSPECTORTOOLSET_API bool WaitFor(const FString& Text = TEXT(""), const FString& TextGone = TEXT(""));
+
+	/** Fill multiple Slate form fields at once.
+	 * @param Fields  Array of {Ref, Value, FieldType} where FieldType is
+	 *                "textbox", "checkbox", or "combobox". */
+	UFUNCTION(BlueprintCallable, Category = "SlateInspectorToolset")
+	static SLATEINSPECTORTOOLSET_API bool FillForm(const TArray<FSlateInspectorToolsetFormField>& Fields);
+
+private:
+
+	/** Resolve a ref to a widget, logging a warning if it fails. */
+	static TSharedPtr<SWidget> ResolveRefOrWarn(const FString& Ref);
+
+	/** Get the absolute screen-space center of a widget. */
+	static FVector2D GetWidgetScreenCenter(TSharedRef<SWidget> Widget);
+
+	/** Map a button name ("left", "right", "middle") to an FKey. */
+	static FKey ParseMouseButtonKey(const FString& Button);
+
+	/** Simulate a mouse click at a widget's screen-space center. */
+	static bool SimulateClick(TSharedRef<SWidget> Widget, FKey MouseButton, bool bDoubleClick, const FModifierKeysState& ModifierKeys = FModifierKeysState());
+};
diff --git a/Engine/Plugins/Experimental/SlateInspectorToolset/Source/SlateInspectorToolset/Public/SlateInspectorToolsetModule.h b/Engine/Plugins/Experimental/SlateInspectorToolset/Source/SlateInspectorToolset/Public/SlateInspectorToolsetModule.h
new file mode 100644
--- /dev/null
+++ b/Engine/Plugins/Experimental/SlateInspectorToolset/Source/SlateInspectorToolset/Public/SlateInspectorToolsetModule.h
@@ -0,0 +1,14 @@
+// Copyright Epic Games, Inc. All Rights Reserved.
+
+#pragma once
+
+#include "Modules/ModuleInterface.h"
+
+class FSlateInspectorToolsetModule : public IModuleInterface
+{
+public:
+
+	virtual void StartupModule() override;
+
+	virtual void ShutdownModule() override;
+};
diff --git a/Engine/Plugins/Experimental/SlateInspectorToolset/Source/SlateInspectorToolset/Public/SlateInspectorToolsetObserverManager.h b/Engine/Plugins/Experimental/SlateInspectorToolset/Source/SlateInspectorToolset/Public/SlateInspectorToolsetObserverManager.h
new file mode 100644
--- /dev/null
+++ b/Engine/Plugins/Experimental/SlateInspectorToolset/Source/SlateInspectorToolset/Public/SlateInspectorToolsetObserverManager.h
@@ -0,0 +1,96 @@
+// Copyright Epic Games, Inc. All Rights Reserved.
+
+#pragma once
+
+#include "CoreMinimal.h"
+#include "Widgets/SWidget.h"
+
+/** A single observer watching a Slate widget subtree. */
+struct FSlateInspectorToolsetObserver
+{
+	/** Unique identifier returned by AddObserver(). */
+	FString Identifier;
+
+	/** Root widget to observe. Invalid (null) when bRoot is true. */
+	TWeakPtr<SWidget> RootWidget;
+
+	/** When true, this is the root observer covering all visible windows (RootWidget is ignored). */
+	bool bRoot = false;
+
+	/** Maximum depth to walk from the root. */
+	int32 MaxDepth = 30;
+
+	/** Cached snapshot text from the most recent observer tick. */
+	FString CachedSnapshotText;
+};
+
+/**
+ * Manages a set of observers that continuously walk Slate widget subtrees
+ * on a throttled tick, keeping the ref cache up to date as the UI changes.
+ *
+ * A default root observer (all visible windows) is always present.
+ */
+class FSlateInspectorToolsetObserverManager
+{
+public:
+
+	SLATEINSPECTORTOOLSET_API static FSlateInspectorToolsetObserverManager& Get();
+
+	/**
+	 * Register an observer on a widget subtree.
+	 *
+	 * @param RootWidget  The root of the subtree to observe.  Null = all visible windows.
+	 * @param MaxDepth    Maximum depth to walk.
+	 * @return The observer identifier (used with RemoveObserver and GetCachedSnapshot).
+	 */
+	SLATEINSPECTORTOOLSET_API FString AddObserver(TSharedPtr<SWidget> RootWidget, int32 MaxDepth = 30);
+
+	/**
+	 * Remove an observer by identifier.
+	 *
+	 * @return True if the observer was found and removed.
+	 */
+	SLATEINSPECTORTOOLSET_API bool RemoveObserver(const FString& Identifier);
+
+	/** Returns a copy of all active observers. */
+	SLATEINSPECTORTOOLSET_API TArray<FSlateInspectorToolsetObserver> GetObservers() const;
+
+	/**
+	 * Returns the cached snapshot text for an observer.
+	 *
+	 * @return The most recently rendered snapshot text, or empty if not found.
+	 */
+	SLATEINSPECTORTOOLSET_API FString GetCachedSnapshot(const FString& Identifier) const;
+
+	/**
+	 * Find the best matching observer for the given root and return its cached snapshot.
+	 *
+	 * A null RootWidget matches the root observer.  Otherwise, matches
+	 * an observer whose root is the given widget (preferring deeper MaxDepth).
+	 *
+	 * @param RootWidget  The root to match. Null = root observer.
+	 * @param MaxDepth    Only match observers with at least this depth.
+	 * @return The cached snapshot text, or empty if no matching observer found.
+	 */
+	SLATEINSPECTORTOOLSET_API FString FindMatchingObserverSnapshot(TSharedPtr<SWidget> RootWidget, int32 MaxDepth) const;
+
+	/**
+	 * Called each Slate post-tick.  Throttled internally to ~100ms intervals.
+	 */
+	SLATEINSPECTORTOOLSET_API void Tick(float DeltaTime);
+
+private:
+
+	TArray<FSlateInspectorToolsetObserver> Observers;
+
+	int32 NextIdentifier = 1;
+
+	/** Accumulated time since last observer pass. */
+	float AccumulatedTime = 0.0f;
+
+	/** Tick counter for periodic purge scheduling. */
+	int32 TickCounter = 0;
+
+	static constexpr float TickIntervalSeconds = 0.1f;
+	static constexpr int32 PurgeEveryNTicks = 10;
+};
diff --git a/Engine/Plugins/Experimental/SlateInspectorToolset/Source/SlateInspectorToolset/Public/SlateInspectorToolsetRefCache.h b/Engine/Plugins/Experimental/SlateInspectorToolset/Source/SlateInspectorToolset/Public/SlateInspectorToolsetRefCache.h
new file mode 100644
--- /dev/null
+++ b/Engine/Plugins/Experimental/SlateInspectorToolset/Source/SlateInspectorToolset/Public/SlateInspectorToolsetRefCache.h
@@ -0,0 +1,86 @@
+// Copyright Epic Games, Inc. All Rights Reserved.
+
+#pragma once
+
+#include "CoreMinimal.h"
+#include "Widgets/SWidget.h"
+
+/**
+ * Manages bidirectional mapping between SWidget instances and short reference
+ * strings used by the SlateInspectorToolset snapshot and action tools.
+ *
+ * # Observer model (append-only)
+ *
+ * When observers are active, refs are never reset.  New widgets encountered
+ * during observer ticks receive the next available counter value for their
+ * role prefix.  Existing widgets keep their refs indefinitely (the same
+ * SWidget pointer always maps to the same ref string).  Destroyed widgets
+ * are cleaned up periodically via PurgeExpired().
+ *
+ * Refs are no longer dense after widgets are destroyed (you may see b1, b3, b7
+ * instead of b1, b2, b3), but keeping refs alive for existing widgets
+ * matters more than compactness.
+ *
+ * Reset() is available for callers that need to clear the cache entirely
+ * and start fresh (e.g. testing).
+ */
+class FSlateInspectorToolsetRefCache
+{
+public:
+
+	SLATEINSPECTORTOOLSET_API static FSlateInspectorToolsetRefCache& Get();
+
+	/**
+	 * Returns existing ref or assigns a new one using the role prefix.
+	 *
+	 * The role prefix determines the ref format: "b" for buttons produces
+	 * b1, b2, b3, etc.  Within a single snapshot pass, the same widget
+	 * always gets the same ref (deduplication via WidgetToRef lookup).
+	 *
+	 * @param Widget      The widget to assign a ref to.
+	 * @param RolePrefix  Short prefix derived from the widget's role (e.g., "b", "tb", "cb").
+	 * @return The assigned ref string (e.g., "b3").
+	 */
+	SLATEINSPECTORTOOLSET_API FString GetOrAssignRef(TSharedRef<SWidget> Widget, const FString& RolePrefix);
+
+	/**
+	 * Resolves a ref string back to a live widget.
+	 *
+	 * Returns nullptr if the ref is unknown or the widget has been destroyed
+	 * since the last snapshot.
+	 */
+	SLATEINSPECTORTOOLSET_API TSharedPtr<SWidget> ResolveRef(const FString& Ref) const;
+
+	/**
+	 * Looks up the ref for a widget without assigning one.
+	 *
+	 * @return The ref string, or empty if the widget has no ref.
+	 */
+	SLATEINSPECTORTOOLSET_API FString FindRef(TSharedRef<SWidget> Widget) const;
+
+	/**
+	 * Clears all mappings and resets counters to zero.
+	 * Observer-driven passes should NOT call this because they rely on
+	 * the append-only behavior to preserve existing refs.
+	 */
+	SLATEINSPECTORTOOLSET_API void Reset();
+
+	/**
+	 * Removes entries whose widget has been destroyed (weak pointer expired).
+	 *
+	 * Called periodically by the observer manager to prevent unbounded
+	 * map growth.  Safe to call at any time.
+	 */
+	SLATEINSPECTORTOOLSET_API void PurgeExpired();
+
+private:
+
+	/** Widget -> ref string (deduplication within a single snapshot pass). */
+	TMap<TWeakPtr<SWidget>, FString> WidgetToRef;
+
+	/** Ref string -> widget (reverse lookup for action tools). */
+	TMap<FString, TWeakPtr<SWidget>> RefToWidget;
+
+	/** Per-role-prefix counters for generating sequential ref numbers. */
+	TMap<FString, int32> RolePrefixCounters;
+};
diff --git a/Engine/Plugins/Experimental/SlateInspectorToolset/Source/SlateInspectorToolset/Public/SlateInspectorToolsetSnapshotRenderer.h b/Engine/Plugins/Experimental/SlateInspectorToolset/Source/SlateInspectorToolset/Public/SlateInspectorToolsetSnapshotRenderer.h
new file mode 100644
--- /dev/null
+++ b/Engine/Plugins/Experimental/SlateInspectorToolset/Source/SlateInspectorToolset/Public/SlateInspectorToolsetSnapshotRenderer.h
@@ -0,0 +1,83 @@
+// Copyright Epic Games, Inc. All Rights Reserved.
+
+#pragma once
+
+#include "CoreMinimal.h"
+#include "Widgets/SWidget.h"
+
+/**
+ * Renders a text-based accessibility snapshot of a Slate widget tree.
+ *
+ * Walks the tree depth-first, emitting one line per meaningful widget with
+ * role, label, state flags, and a compact ref identifier.  Structural
+ * containers (SBox, SBorder, SOverlay, etc.) are elided and their children
+ * are emitted at the parent's indent level.
+ */
+class FSlateInspectorToolsetSnapshotRenderer
+{
+public:
+
+	/**
+	 * Render a snapshot starting from the given root widget.
+	 *
+	 * @param Root                     The widget subtree to snapshot.  If null, snapshots all top-level windows.
+	 * @param MaxDepth                 Maximum recursion depth (default 30).
+	 * @param bIncludeSourceLocations  Include [src=File:Line] tags showing where each widget was created in C++.
+	 * @param bResetCache              When true (default), clears the ref cache before rendering.
+	 *                                 Pass false for observer-driven passes that preserve existing refs.
+	 * @return Indented text representation of the widget tree.
+	 */
+	SLATEINSPECTORTOOLSET_API static FString Render(TSharedPtr<SWidget> Root, int32 MaxDepth = 30, bool bIncludeSourceLocations = false, bool bResetCache = true);
+
+	/** Register a custom widget type so it appears in snapshots. */
+	SLATEINSPECTORTOOLSET_API static void RegisterWidgetRole(const FName& WidgetType, const FString& Role, const FString& RefPrefix);
+
+	/** Register a custom label extractor for a widget type. */
+	SLATEINSPECTORTOOLSET_API static void RegisterLabelExtractor(const FName& WidgetType, TFunction<FString(TSharedRef<SWidget>)> Extractor);
+
+	/** Register a custom state flags extractor for a widget type. */
+	SLATEINSPECTORTOOLSET_API static void RegisterStateFlagsExtractor(const FName& WidgetType, TFunction<TArray<FString>(TSharedRef<SWidget>)> Extractor);
+
+private:
+
+	/** Initializes the type-to-role and structural-skip maps on first use. */
+	static void EnsureMapsInitialized();
+
+	/** Recursive depth-first walk. */
+	static void RenderWidget(
+		FStringBuilderBase& Output,
+		TSharedRef<SWidget> Widget,
+		int32 IndentLevel,
+		int32 CurrentDepth,
+		int32 MaxDepth,
+		bool bIncludeSourceLocations);
+
+	/** Extract a human-readable label from a widget via type-specific methods. */
+	static FString ExtractLabel(TSharedRef<SWidget> Widget, const FString& TypeString);
+
+	/** Extract state flags (visibility, enabled, checked, focused). */
+	static FString ExtractStateFlags(TSharedRef<SWidget> Widget, const FString& TypeString);
+
+	/** Look up role for a widget type, handling template types via prefix matching. */
+	static const FString* FindRoleForType(const FName& WidgetType, const FString& TypeString);
+
+	/** Map from widget type name to role string (e.g., "SButton" -> "button"). */
+	static TMap<FName, FString> TypeToRole;
+
+	/** Prefix-based role mapping for template types (e.g., "SComboBox" matches "SComboBox<...>"). */
+	static TArray<TPair<FString, FString>> TemplatePrefixToRole;
+
+	/** Map from role string to ref prefix (e.g., "button" -> "b"). */
+	static TMap<FString, FString> RoleToRefPrefix;
+
+	/** Set of structural container type names that should be skipped (children promoted). */
+	static TSet<FName> StructuralContainers;
+
+	/** Custom label extractors registered via RegisterLabelExtractor. */
+	static TMap<FName, TFunction<FString(TSharedRef<SWidget>)>> CustomLabelExtractors;
+
+	/** Custom state flags extractors registered via RegisterStateFlagsExtractor. */
+	static TMap<FName, TFunction<TArray<FString>(TSharedRef<SWidget>)>> CustomStateFlagsExtractors;
+
+	static bool bMapsInitialized;
+};
diff --git a/Engine/Plugins/Experimental/SlateInspectorToolset/Source/SlateInspectorToolset/Public/SlateInspectorToolsetSubsystem.h b/Engine/Plugins/Experimental/SlateInspectorToolset/Source/SlateInspectorToolset/Public/SlateInspectorToolsetSubsystem.h
new file mode 100644
--- /dev/null
+++ b/Engine/Plugins/Experimental/SlateInspectorToolset/Source/SlateInspectorToolset/Public/SlateInspectorToolsetSubsystem.h
@@ -0,0 +1,35 @@
+// Copyright Epic Games, Inc. All Rights Reserved.
+
+#pragma once
+
+#include "EditorSubsystem.h"
+
+#include "SlateInspectorToolsetSubsystem.generated.h"
+
+class FSubsystemCollectionBase;
+
+UCLASS(MinimalAPI)
+class USlateInspectorToolsetSubsystem : public UEditorSubsystem
+{
+	GENERATED_BODY()
+
+public:
+
+	virtual void Initialize(FSubsystemCollectionBase& Collection) override;
+
+	virtual void Deinitialize() override;
+
+	/** Register or unregister the toolset at runtime. */
+	void SetToolsetEnabled(bool bEnabled);
+
+private:
+
+	/** Forwarded from FSlateApplication::OnPostTick to drive observer updates. */
+	void OnSlatePostTick(float DeltaTime);
+
+	bool bToolsetRegistered = false;
+
+	FDelegateHandle PostTickDelegateHandle;
+
+	FString RootObserverIdentifier;
+};
diff --git a/Engine/Plugins/Experimental/SlateInspectorToolset/Source/SlateInspectorToolset/SlateInspectorToolset.Build.cs b/Engine/Plugins/Experimental/SlateInspectorToolset/Source/SlateInspectorToolset/SlateInspectorToolset.Build.cs
new file mode 100644
--- /dev/null
+++ b/Engine/Plugins/Experimental/SlateInspectorToolset/Source/SlateInspectorToolset/SlateInspectorToolset.Build.cs
@@ -0,0 +1,30 @@
+// Copyright Epic Games, Inc. All Rights Reserved.
+
+namespace UnrealBuildTool.Rules
+{
+	public class SlateInspectorToolset : ModuleRules
+	{
+		public SlateInspectorToolset(ReadOnlyTargetRules Target) : base(Target)
+		{
+			PublicDependencyModuleNames.AddRange(
+				new string[] {
+					"ToolsetRegistry",
+				}
+			);
+
+			PrivateDependencyModuleNames.AddRange(
+				new string[] {
+					"Core",
+					"CoreUObject",
+					"EditorSubsystem",
+					"Engine",
+					"ImageWrapper",
+					"InputCore",
+					"Slate",
+					"SlateCore",
+					"UnrealEd",
+				}
+			);
+		}
+	}
+}
diff --git a/Engine/Plugins/Experimental/SlateInspectorToolset/Source/SlateInspectorToolsetTests/Private/SlateInspectorToolsetTestPanel.cpp b/Engine/Plugins/Experimental/SlateInspectorToolset/Source/SlateInspectorToolsetTests/Private/SlateInspectorToolsetTestPanel.cpp
new file mode 100644
--- /dev/null
+++ b/Engine/Plugins/Experimental/SlateInspectorToolset/Source/SlateInspectorToolsetTests/Private/SlateInspectorToolsetTestPanel.cpp
@@ -0,0 +1,125 @@
+// Copyright Epic Games, Inc. All Rights Reserved.
+
+#include "SlateInspectorToolsetTestPanel.h"
+
+#include "Widgets/Input/SButton.h"
+#include "Widgets/Layout/SBox.h"
+#include "Widgets/SBoxPanel.h"
+#include "Widgets/Text/STextBlock.h"
+
+void SSlateInspectorToolsetTestPanel::Construct(const FArguments& InArgs)
+{
+	ComboOptions.Add(MakeShared<FString>(TEXT("Alpha")));
+	ComboOptions.Add(MakeShared<FString>(TEXT("Beta")));
+	ComboOptions.Add(MakeShared<FString>(TEXT("Gamma")));
+	SelectedComboOption = ComboOptions[0];
+
+	ChildSlot
+	[
+		SNew(SVerticalBox)
+
+		// Primary button
+		+ SVerticalBox::Slot().AutoHeight().Padding(4)
+		[
+			SNew(SButton)
+			.Text(FText::FromString(TEXT("Test Button")))
+			.OnClicked_Lambda([this]()
+			{
+				++ButtonClickCount;
+				return FReply::Handled();
+			})
+		]
+
+		// Second button (for double-click testing)
+		+ SVerticalBox::Slot().AutoHeight().Padding(4)
+		[
+			SNew(SButton)
+			.Text(FText::FromString(TEXT("Second Button")))
+			.OnClicked_Lambda([this]()
+			{
+				++SecondButtonClickCount;
+				return FReply::Handled();
+			})
+		]
+
+		// Single-line text box
+		+ SVerticalBox::Slot().AutoHeight().Padding(4)
+		[
+			SAssignNew(TextBox, SEditableTextBox)
+			.HintText(FText::FromString(TEXT("Type here")))
+		]
+
+		// Multi-line text box
+		+ SVerticalBox::Slot().AutoHeight().Padding(4)
+		[
+			SNew(SBox)
+			.HeightOverride(60)
+			[
+				SAssignNew(MultiLineTextBox, SMultiLineEditableTextBox)
+				.HintText(FText::FromString(TEXT("Multi-line")))
+			]
+		]
+
+		// Search box
+		+ SVerticalBox::Slot().AutoHeight().Padding(4)
+		[
+			SAssignNew(SearchBox, SSearchBox)
+		]
+
+		// Checkbox
+		+ SVerticalBox::Slot().AutoHeight().Padding(4)
+		[
+			SAssignNew(CheckBox, SCheckBox)
+		]
+
+		// Combobox
+		+ SVerticalBox::Slot().AutoHeight().Padding(4)
+		[
+			SAssignNew(ComboBox, SComboBox<TSharedPtr<FString>>)
+			.OptionsSource(&ComboOptions)
+			.OnSelectionChanged_Lambda([this](TSharedPtr<FString> NewValue, ESelectInfo::Type)
+			{
+				SelectedComboOption = NewValue;
+			})
+			.OnGenerateWidget_Lambda([](TSharedPtr<FString> Item)
+			{
+				return SNew(STextBlock).Text(FText::FromString(*Item));
+			})
+			[
+				SNew(STextBlock)
+				.Text_Lambda([this]() { return FText::FromString(*SelectedComboOption); })
+			]
+		]
+
+		// Slider
+		+ SVerticalBox::Slot().AutoHeight().Padding(4)
+		[
+			SAssignNew(Slider, SSlider)
+			.Value(0.0f)
+			.OnValueChanged_Lambda([this](float NewValue)
+			{
+				SliderValue = NewValue;
+			})
+		]
+
+		// SpinBox
+		+ SVerticalBox::Slot().AutoHeight().Padding(4)
+		[
+			SAssignNew(SpinBox, SSpinBox<float>)
+			.MinValue(0.0f)
+			.MaxValue(100.0f)
+			.Value(0.0f)
+			.OnValueChanged_Lambda([this](float NewValue)
+			{
+				SpinBoxValue = NewValue;
+			})
+		]
+
+		// Static text label
+		+ SVerticalBox::Slot().AutoHeight().Padding(4)
+		[
+			SNew(STextBlock)
+			.Text(FText::FromString(TEXT("Static Label")))
+		]
+	];
+}
diff --git a/Engine/Plugins/Experimental/SlateInspectorToolset/Source/SlateInspectorToolsetTests/Private/SlateInspectorToolsetTestPanel.h b/Engine/Plugins/Experimental/SlateInspectorToolset/Source/SlateInspectorToolsetTests/Private/SlateInspectorToolsetTestPanel.h
new file mode 100644
--- /dev/null
+++ b/Engine/Plugins/Experimental/SlateInspectorToolset/Source/SlateInspectorToolsetTests/Private/SlateInspectorToolsetTestPanel.h
@@ -0,0 +1,50 @@
+// Copyright Epic Games, Inc. All Rights Reserved.
+
+#pragma once
+
+#include "Widgets/SCompoundWidget.h"
+#include "Widgets/Input/SCheckBox.h"
+#include "Widgets/Input/SComboBox.h"
+#include "Widgets/Input/SEditableTextBox.h"
+#include "Widgets/Input/SSlider.h"
+#include "Widgets/Input/SSpinBox.h"
+#include "Widgets/Input/SSearchBox.h"
+#include "Widgets/Input/SMultiLineEditableTextBox.h"
+
+/**
+ * A widget panel containing one of each common widget type.
+ * Used by SlateInspectorToolset automation tests to exercise snapshot, observer,
+ * and input simulation tools against known widgets.
+ */
+class SSlateInspectorToolsetTestPanel : public SCompoundWidget
+{
+public:
+
+	SLATE_BEGIN_ARGS(SSlateInspectorToolsetTestPanel) {}
+	SLATE_END_ARGS()
+
+	void Construct(const FArguments& InArgs);
+
+	// Button
+	int32 ButtonClickCount = 0;
+	int32 SecondButtonClickCount = 0;
+
+	// Text input
+	TSharedPtr<SEditableTextBox> TextBox;
+	TSharedPtr<SMultiLineEditableTextBox> MultiLineTextBox;
+	TSharedPtr<SSearchBox> SearchBox;
+
+	// Toggle
+	TSharedPtr<SCheckBox> CheckBox;
+
+	// Combobox
+	TSharedPtr<SComboBox<TSharedPtr<FString>>> ComboBox;
+	TArray<TSharedPtr<FString>> ComboOptions;
+	TSharedPtr<FString> SelectedComboOption;
+
+	// Slider / SpinBox
+	TSharedPtr<SSlider> Slider;
+	float SliderValue = 0.0f;
+	TSharedPtr<SSpinBox<float>> SpinBox;
+	float SpinBoxValue = 0.0f;
+};
diff --git a/Engine/Plugins/Experimental/SlateInspectorToolset/Source/SlateInspectorToolsetTests/Private/SlateInspectorToolsetTests.cpp b/Engine/Plugins/Experimental/SlateInspectorToolset/Source/SlateInspectorToolsetTests/Private/SlateInspectorToolsetTests.cpp
new file mode 100644
--- /dev/null
+++ b/Engine/Plugins/Experimental/SlateInspectorToolset/Source/SlateInspectorToolsetTests/Private/SlateInspectorToolsetTests.cpp
@@ -0,0 +1,776 @@
+// Copyright Epic Games, Inc. All Rights Reserved.
+
+#include "CQTest.h"
+#include "Components/CQTestSlateComponent.h"
+#include "SlateInspectorToolsetTestPanel.h"
+
+#include "Framework/Application/SlateApplication.h"
+#include "SlateInspectorToolsetObserverManager.h"
+#include "SlateInspectorToolsetRefCache.h"
+#include "SlateInspectorToolset.h"
+#include "Widgets/SWindow.h"
+
+TEST_CLASS(SlateInspectorToolsetTest, "Editor.SlateInspectorToolset")
+{
+	TUniquePtr<FCQTestSlateComponent> SlateComponent;
+	TSharedPtr<SWindow> TestWindow;
+	TSharedPtr<SSlateInspectorToolsetTestPanel> TestPanel;
+
+	BEFORE_EACH()
+	{
+		SlateComponent = MakeUnique<FCQTestSlateComponent>();
+
+		TestPanel = SNew(SSlateInspectorToolsetTestPanel);
+
+		TestWindow = SNew(SWindow)
+			.Title(FText::FromString(TEXT("SlateInspectorToolset Test Window")))
+			.ScreenPosition(FVector2D(100, 100))
+			.ClientSize(FVector2D(400, 600))
+			.SupportsMaximize(false)
+			.SupportsMinimize(false)
+			.FocusWhenFirstShown(true)
+			[
+				TestPanel.ToSharedRef()
+			];
+
+		FSlateApplication::Get().AddWindow(TestWindow.ToSharedRef());
+		TestWindow->MoveWindowTo(FVector2D(100, 100));
+		TestWindow->BringToFront();
+
+		// Force layout so widgets have valid geometry before any interaction.
+		FSlateApplication::Get().Tick();
+	}
+
+	AFTER_EACH()
+	{
+		// Clean up any observers created during the test (except the root observer).
+		for (const FSlateInspectorToolsetObserver& Observer : FSlateInspectorToolsetObserverManager::Get().GetObservers())
+		{
+			if (!Observer.bRoot)
+			{
+				FSlateInspectorToolsetObserverManager::Get().RemoveObserver(Observer.Identifier);
+			}
+		}
+
+		if (TestWindow.IsValid())
+		{
+			TestWindow->RequestDestroyWindow();
+			TestWindow.Reset();
+		}
+		TestPanel.Reset();
+		SlateComponent.Reset();
+	}
+
+	// ── Snapshot tests ──────────────────────────────────────────────
+
+	TEST_METHOD(Snapshot_ReturnsWidgetTree)
+	{
+		const FString Snapshot = SnapshotTestWindow();
+		ASSERT_THAT(IsTrue(!Snapshot.IsEmpty()));
+		ASSERT_THAT(IsTrue(Snapshot.Contains(TEXT("button"))));
+		ASSERT_THAT(IsTrue(Snapshot.Contains(TEXT("textbox"))));
+		ASSERT_THAT(IsTrue(Snapshot.Contains(TEXT("checkbox"))));
+		ASSERT_THAT(IsTrue(Snapshot.Contains(TEXT("combobox"))));
+		ASSERT_THAT(IsTrue(Snapshot.Contains(TEXT("slider"))));
+	}
+
+	TEST_METHOD(Snapshot_SubtreeByRef)
+	{
+		const FString FullSnapshot = USlateInspectorToolset::Snapshot(TEXT(""), 30);
+		ASSERT_THAT(IsTrue(FullSnapshot.Contains(TEXT("SlateInspectorToolset Test Window"))));
+
+		const FString WindowRef = FindRefByRole(FullSnapshot, TEXT("window"), TEXT("SlateInspectorToolset Test Window"));
+		ASSERT_THAT(IsTrue(!WindowRef.IsEmpty()));
+
+		const FString SubtreeSnapshot = USlateInspectorToolset::Snapshot(WindowRef, 30);
+		ASSERT_THAT(IsTrue(SubtreeSnapshot.Contains(TEXT("button"))));
+	}
+
+	TEST_METHOD(Snapshot_ContainsStaticLabel)
+	{
+		const FString Snapshot = SnapshotTestWindow();
+		ASSERT_THAT(IsTrue(Snapshot.Contains(TEXT("Static Label"))));
+	}
+
+	TEST_METHOD(Snapshot_ContainsSpinBox)
+	{
+		const FString Snapshot = SnapshotTestWindow();
+		// SSpinBox<> maps to "slider" role via template prefix matching.
+		ASSERT_THAT(IsTrue(Snapshot.Contains(TEXT("slider"))));
+	}
+
+	// ── Observer tests ──────────────────────────────────────────────
+
+	TEST_METHOD(Observe_CreatesObserver)
+	{
+		const FString Snapshot = SnapshotTestWindow();
+		const FString WindowRef = FindRefByRole(Snapshot, TEXT("window"), TEXT("SlateInspectorToolset Test Window"));
+		ASSERT_THAT(IsTrue(!WindowRef.IsEmpty()));
+
+		const FString Identifier = USlateInspectorToolset::Observe(WindowRef, 30);
+		ASSERT_THAT(IsTrue(!Identifier.IsEmpty()));
+		ASSERT_THAT(IsTrue(!Identifier.StartsWith(TEXT("Error"))));
+
+		const FString ObserverList = USlateInspectorToolset::ListObservers();
+		ASSERT_THAT(IsTrue(ObserverList.Contains(Identifier)));
+
+		USlateInspectorToolset::Unobserve(Identifier);
+	}
+
+	TEST_METHOD(Observe_Deduplicates)
+	{
+		const FString Snapshot = SnapshotTestWindow();
+		const FString WindowRef = FindRefByRole(Snapshot, TEXT("window"), TEXT("SlateInspectorToolset Test Window"));
+
+		const FString Identifier1 = USlateInspectorToolset::Observe(WindowRef, 30);
+		const FString Identifier2 = USlateInspectorToolset::Observe(WindowRef, 30);
+		ASSERT_THAT(IsTrue(Identifier1 == Identifier2));
+
+		USlateInspectorToolset::Unobserve(Identifier1);
+	}
+
+	TEST_METHOD(Unobserve_RemovesObserver)
+	{
+		const FString Snapshot = SnapshotTestWindow();
+		const FString WindowRef = FindRefByRole(Snapshot, TEXT("window"), TEXT("SlateInspectorToolset Test Window"));
+		const FString Identifier = USlateInspectorToolset::Observe(WindowRef, 30);
+
+		const bool bRemoved = USlateInspectorToolset::Unobserve(Identifier);
+		ASSERT_THAT(IsTrue(bRemoved));
+
+		const FString ObserverList = USlateInspectorToolset::ListObservers();
+		ASSERT_THAT(IsTrue(!ObserverList.Contains(Identifier)));
+	}
+
+	TEST_METHOD(Unobserve_ReturnsFalseForUnknown)
+	{
+		const bool bRemoved = USlateInspectorToolset::Unobserve(TEXT("nonexistent_observer"));
+		ASSERT_THAT(IsTrue(!bRemoved));
+	}
+
+	TEST_METHOD(Observer_KeepsRefsAlive)
+	{
+		const FString Snapshot = SnapshotTestWindow();
+		const FString ButtonRef = FindRefByRole(Snapshot, TEXT("button"), TEXT("Test Button"));
+		ASSERT_THAT(IsTrue(!ButtonRef.IsEmpty()));
+
+		const FString WindowRef = FindRefByRole(Snapshot, TEXT("window"), TEXT("SlateInspectorToolset Test Window"));
+		const FString ObserverId = USlateInspectorToolset::Observe(WindowRef, 30);
+
+		TestCommandBuilder
+			.StartWhen([this]() { return SlateComponent->HaveTicksElapsed(15); })
+			.Then([this, ButtonRef, ObserverId]()
+			{
+				const bool bClicked = USlateInspectorToolset::Click(ButtonRef);
+				ASSERT_THAT(IsTrue(bClicked));
+				ASSERT_THAT(IsTrue(TestPanel->ButtonClickCount > 0));
+				USlateInspectorToolset::Unobserve(ObserverId);
+			});
+	}
+
+	// ── Click tests ─────────────────────────────────────────────────
+
+	TEST_METHOD(Click_PressesButton)
+	{
+		TestCommandBuilder
+			.StartWhen([this]() { return SlateComponent->HaveTicksElapsed(3); })
+			.Then([this]()
+			{
+				const FString Snapshot = SnapshotTestWindow();
+				const FString ButtonRef = FindRefByRole(Snapshot, TEXT("button"), TEXT("Test Button"));
+				ASSERT_THAT(IsTrue(!ButtonRef.IsEmpty()));
+
+				ASSERT_THAT(IsTrue(TestPanel->ButtonClickCount == 0));
+				const bool bClicked = USlateInspectorToolset::Click(ButtonRef);
+				ASSERT_THAT(IsTrue(bClicked));
+				ASSERT_THAT(IsTrue(TestPanel->ButtonClickCount == 1));
+			});
+	}
+
+	TEST_METHOD(Click_MultipleButtons)
+	{
+		TestCommandBuilder
+			.StartWhen([this]() { return SlateComponent->HaveTicksElapsed(3); })
+			.Then([this]()
+			{
+				const FString Snapshot = SnapshotTestWindow();
+				const FString FirstRef = FindRefByRole(Snapshot, TEXT("button"), TEXT("Test Button"));
+				const FString SecondRef = FindRefByRole(Snapshot, TEXT("button"), TEXT("Second Button"));
+				ASSERT_THAT(IsTrue(!FirstRef.IsEmpty()));
+				ASSERT_THAT(IsTrue(!SecondRef.IsEmpty()));
+				ASSERT_THAT(IsTrue(FirstRef != SecondRef));
+
+				USlateInspectorToolset::Click(FirstRef);
+				USlateInspectorToolset::Click(SecondRef);
+				USlateInspectorToolset::Click(FirstRef);
+
+				ASSERT_THAT(IsTrue(TestPanel->ButtonClickCount == 2));
+				ASSERT_THAT(IsTrue(TestPanel->SecondButtonClickCount == 1));
+			});
+	}
+
+	TEST_METHOD(Click_DoubleClick)
+	{
+		TestCommandBuilder
+			.StartWhen([this]() { return SlateComponent->HaveTicksElapsed(3); })
+			.Then([this]()
+			{
+				const FString Snapshot = SnapshotTestWindow();
+				const FString ButtonRef = FindRefByRole(Snapshot, TEXT("button"), TEXT("Test Button"));
+				ASSERT_THAT(IsTrue(!ButtonRef.IsEmpty()));
+
+				const bool bClicked = USlateInspectorToolset::Click(ButtonRef, TEXT("left"), true);
+				ASSERT_THAT(IsTrue(bClicked));
+			});
+	}
+
+	TEST_METHOD(Click_RightClick)
+	{
+		TestCommandBuilder
+			.StartWhen([this]() { return SlateComponent->HaveTicksElapsed(3); })
+			.Then([this]()
+			{
+				const FString Snapshot = SnapshotTestWindow();
+				const FString ButtonRef = FindRefByRole(Snapshot, TEXT("button"), TEXT("Test Button"));
+				ASSERT_THAT(IsTrue(!ButtonRef.IsEmpty()));
+
+				const bool bClicked = USlateInspectorToolset::Click(ButtonRef, TEXT("right"));
+				ASSERT_THAT(IsTrue(bClicked));
+			});
+	}
+
+	TEST_METHOD(Click_InvalidRef)
+	{
+		const bool bClicked = USlateInspectorToolset::Click(TEXT("nonexistent_ref"));
+		ASSERT_THAT(IsTrue(!bClicked));
+	}
+
+	TEST_METHOD(Click_WithShiftModifier)
+	{
+		TestCommandBuilder
+			.StartWhen([this]() { return SlateComponent->HaveTicksElapsed(3); })
+			.Then([this]()
+			{
+				const FString Snapshot = SnapshotTestWindow();
+				const FString ButtonRef = FindRefByRole(Snapshot, TEXT("button"), TEXT("Test Button"));
+				ASSERT_THAT(IsTrue(!ButtonRef.IsEmpty()));
+
+				FSlateInspectorToolsetModifierKeys Modifiers;
+				Modifiers.bShift = true;
+				const bool bClicked = USlateInspectorToolset::Click(ButtonRef, TEXT("left"), false, Modifiers);
+				ASSERT_THAT(IsTrue(bClicked));
+				ASSERT_THAT(IsTrue(TestPanel->ButtonClickCount == 1));
+			});
+	}
+
+	TEST_METHOD(Click_WithCtrlModifier)
+	{
+		TestCommandBuilder
+			.StartWhen([this]() { return SlateComponent->HaveTicksElapsed(3); })
+			.Then([this]()
+			{
+				const FString Snapshot = SnapshotTestWindow();
+				const FString ButtonRef = FindRefByRole(Snapshot, TEXT("button"), TEXT("Test Button"));
+				ASSERT_THAT(IsTrue(!ButtonRef.IsEmpty()));
+
+				FSlateInspectorToolsetModifierKeys Modifiers;
+				Modifiers.bCtrl = true;
+				const bool bClicked = USlateInspectorToolset::Click(ButtonRef, TEXT("left"), false, Modifiers);
+				ASSERT_THAT(IsTrue(bClicked));
+			});
+	}
+
+	TEST_METHOD(Click_WithMultipleModifiers)
+	{
+		TestCommandBuilder
+			.StartWhen([this]() { return SlateComponent->HaveTicksElapsed(3); })
+			.Then([this]()
+			{
+				const FString Snapshot = SnapshotTestWindow();
+				const FString ButtonRef = FindRefByRole(Snapshot, TEXT("button"), TEXT("Test Button"));
+				ASSERT_THAT(IsTrue(!ButtonRef.IsEmpty()));
+
+				FSlateInspectorToolsetModifierKeys Modifiers;
+				Modifiers.bShift = true;
+				Modifiers.bCtrl = true;
+				Modifiers.bAlt = true;
+				const bool bClicked = USlateInspectorToolset::Click(ButtonRef, TEXT("left"), false, Modifiers);
+				ASSERT_THAT(IsTrue(bClicked));
+			});
+	}
+
+	// ── Drag tests ─────────────────────────────────────────────────
+
+	TEST_METHOD(Drag_BetweenButtons)
+	{
+		TestCommandBuilder
+			.StartWhen([this]() { return SlateComponent->HaveTicksElapsed(3); })
+			.Then([this]()
+			{
+				const FString Snapshot = SnapshotTestWindow();
+				const FString FirstRef = FindRefByRole(Snapshot, TEXT("button"), TEXT("Test Button"));
+				const FString SecondRef = FindRefByRole(Snapshot, TEXT("button"), TEXT("Second Button"));
+				ASSERT_THAT(IsTrue(!FirstRef.IsEmpty()));
+				ASSERT_THAT(IsTrue(!SecondRef.IsEmpty()));
+
+				const bool bDragged = USlateInspectorToolset::Drag(FirstRef, SecondRef);
+				ASSERT_THAT(IsTrue(bDragged));
+			});
+	}
+
+	TEST_METHOD(Drag_WithModifiers)
+	{
+		TestCommandBuilder
+			.StartWhen([this]() { return SlateComponent->HaveTicksElapsed(3); })
+			.Then([this]()
+			{
+				const FString Snapshot = SnapshotTestWindow();
+				const FString FirstRef = FindRefByRole(Snapshot, TEXT("button"), TEXT("Test Button"));
+				const FString SecondRef = FindRefByRole(Snapshot, TEXT("button"), TEXT("Second Button"));
+
+				FSlateInspectorToolsetModifierKeys Modifiers;
+				Modifiers.bShift = true;
+				const bool bDragged = USlateInspectorToolset::Drag(FirstRef, SecondRef, Modifiers);
+				ASSERT_THAT(IsTrue(bDragged));
+			});
+	}
+
+	TEST_METHOD(Drag_InvalidStartRef)
+	{
+		const bool bDragged = USlateInspectorToolset::Drag(TEXT("nonexistent"), TEXT("nonexistent2"));
+		ASSERT_THAT(IsTrue(!bDragged));
+	}
+
+	TEST_METHOD(Drag_SliderChangesValue)
+	{
+		TestCommandBuilder
+			.StartWhen([this]() { return SlateComponent->HaveTicksElapsed(3); })
+			.Then([this]()
+			{
+				ASSERT_THAT(IsTrue(TestPanel->SliderValue == 0.0f));
+
+				const FString Snapshot = SnapshotTestWindow();
+				const TArray<FString> SliderRefs = FindAllRefsByRole(Snapshot, TEXT("slider"));
+				// First slider ref is the SSlider, second is the SSpinBox.
+				ASSERT_THAT(IsTrue(SliderRefs.Num() >= 1));
+				const FString SliderRef = SliderRefs[0];
+
+				// Drag from the slider to one of the buttons (to the right and up).
+				// This should move the slider thumb and change the value.
+				const FString ButtonRef = FindRefByRole(Snapshot, TEXT("button"), TEXT("Test Button"));
+				const bool bDragged = USlateInspectorToolset::Drag(SliderRef, ButtonRef);
+				ASSERT_THAT(IsTrue(bDragged));
+
+				// The slider value should have changed from its initial 0.0.
+				// We can't predict the exact value since it depends on geometry,
+				// but it should no longer be zero.
+				ASSERT_THAT(IsTrue(TestPanel->SliderValue != 0.0f));
+			});
+	}
+
+	// ── Checkbox tests ──────────────────────────────────────────────
+
+	TEST_METHOD(Click_TogglesCheckbox)
+	{
+		TestCommandBuilder
+			.StartWhen([this]() { return SlateComponent->HaveTicksElapsed(3); })
+			.Then([this]()
+			{
+				const FString Snapshot = SnapshotTestWindow();
+				const FString CheckBoxRef = FindRefByRole(Snapshot, TEXT("checkbox"), TEXT(""));
+				ASSERT_THAT(IsTrue(!CheckBoxRef.IsEmpty()));
+
+				ASSERT_THAT(IsTrue(!TestPanel->CheckBox->IsChecked()));
+
+				USlateInspectorToolset::Click(CheckBoxRef);
+				ASSERT_THAT(IsTrue(TestPanel->CheckBox->IsChecked()));
+
+				USlateInspectorToolset::Click(CheckBoxRef);
+				ASSERT_THAT(IsTrue(!TestPanel->CheckBox->IsChecked()));
+			});
+	}
+
+	// ── Type tests ──────────────────────────────────────────────────
+
+	TEST_METHOD(Type_EntersText)
+	{
+		const FString Snapshot = SnapshotTestWindow();
+		const FString TextBoxRef = FindRefByRole(Snapshot, TEXT("textbox"), TEXT(""));
+		ASSERT_THAT(IsTrue(!TextBoxRef.IsEmpty()));
+
+		const bool bTyped = USlateInspectorToolset::Type(TextBoxRef, TEXT("hello"));
+		ASSERT_THAT(IsTrue(bTyped));
+		ASSERT_THAT(IsTrue(TestPanel->TextBox->GetText().ToString() == TEXT("hello")));
+	}
+
+	TEST_METHOD(Type_AppendsToExistingText)
+	{
+		const FString Snapshot = SnapshotTestWindow();
+		const FString TextBoxRef = FindRefByRole(Snapshot, TEXT("textbox"), TEXT(""));
+
+		USlateInspectorToolset::Type(TextBoxRef, TEXT("foo"));
+		USlateInspectorToolset::Type(TextBoxRef, TEXT("bar"));
+		ASSERT_THAT(IsTrue(TestPanel->TextBox->GetText().ToString() == TEXT("foobar")));
+	}
+
+	TEST_METHOD(Type_SpecialCharacters)
+	{
+		const FString Snapshot = SnapshotTestWindow();
+		const FString TextBoxRef = FindRefByRole(Snapshot, TEXT("textbox"), TEXT(""));
+
+		USlateInspectorToolset::Type(TextBoxRef, TEXT("a@b#c$d"));
+		ASSERT_THAT(IsTrue(TestPanel->TextBox->GetText().ToString() == TEXT("a@b#c$d")));
+	}
+
+	TEST_METHOD(Type_WithSubmit)
+	{
+		const FString Snapshot = SnapshotTestWindow();
+		const FString TextBoxRef = FindRefByRole(Snapshot, TEXT("textbox"), TEXT(""));
+		ASSERT_THAT(IsTrue(!TextBoxRef.IsEmpty()));
+
+		const bool bTyped = USlateInspectorToolset::Type(TextBoxRef, TEXT("submitted"), true);
+		ASSERT_THAT(IsTrue(bTyped));
+	}
+
+	TEST_METHOD(Type_IntoSearchBox)
+	{
+		const FString Snapshot = SnapshotTestWindow();
+		// SSearchBox contains an SEditableTextBox internally, so look for a textbox.
+		// The search box is the third textbox in the panel (after TextBox and MultiLineTextBox).
+		const TArray<FString> TextBoxRefs = FindAllRefsByRole(Snapshot, TEXT("textbox"));
+		// We need at least 2 textbox refs (single-line and search box; multi-line may be a different role).
+		ASSERT_THAT(IsTrue(TextBoxRefs.Num() >= 2));
+
+		// Type into the last textbox (search box).
+		const FString SearchRef = TextBoxRefs.Last();
+		const bool bTyped = USlateInspectorToolset::Type(SearchRef, TEXT("search query"));
+		ASSERT_THAT(IsTrue(bTyped));
+	}
+
+	TEST_METHOD(Type_InvalidRef)
+	{
+		const bool bTyped = USlateInspectorToolset::Type(TEXT("nonexistent_ref"), TEXT("hello"));
+		ASSERT_THAT(IsTrue(!bTyped));
+	}
+
+	// ── PressKey in text context ────────────────────────────────────
+
+	TEST_METHOD(PressKey_BackspaceDeletesCharacter)
+	{
+		TestCommandBuilder
+			.StartWhen([this]() { return SlateComponent->HaveTicksElapsed(3); })
+			.Then([this]()
+			{
+				const FString Snapshot = SnapshotTestWindow();
+				const FString TextBoxRef = FindRefByRole(Snapshot, TEXT("textbox"), TEXT(""));
+
+				USlateInspectorToolset::Type(TextBoxRef, TEXT("abc"));
+				ASSERT_THAT(IsTrue(TestPanel->TextBox->GetText().ToString() == TEXT("abc")));
+
+				// BackSpace should delete the last character.
+				// Re-focus via Type to ensure the editable text has focus.
+				USlateInspectorToolset::Type(TextBoxRef, TEXT(""));
+				USlateInspectorToolset::PressKey(TEXT("BackSpace"));
+				ASSERT_THAT(IsTrue(TestPanel->TextBox->GetText().ToString() == TEXT("ab")));
+			});
+	}
+
+	TEST_METHOD(PressKey_SelectAllAndDelete)
+	{
+		TestCommandBuilder
+			.StartWhen([this]() { return SlateComponent->HaveTicksElapsed(3); })
+			.Then([this]()
+			{
+				const FString Snapshot = SnapshotTestWindow();
+				const FString TextBoxRef = FindRefByRole(Snapshot, TEXT("textbox"), TEXT(""));
+
+				USlateInspectorToolset::Type(TextBoxRef, TEXT("hello world"));
+				ASSERT_THAT(IsTrue(TestPanel->TextBox->GetText().ToString() == TEXT("hello world")));
+
+				// Ctrl+A to select all, then Delete to clear.
+				USlateInspectorToolset::PressKey(TEXT("Ctrl+A"));
+				USlateInspectorToolset::PressKey(TEXT("Delete"));
+				ASSERT_THAT(IsTrue(TestPanel->TextBox->GetText().ToString().IsEmpty()));
+			});
+	}
+
+	// ── Hover tests ─────────────────────────────────────────────────
+
+	TEST_METHOD(Hover_ReturnsTrue)
+	{
+		TestCommandBuilder
+			.StartWhen([this]() { return SlateComponent->HaveTicksElapsed(3); })
+			.Then([this]()
+			{
+				const FString Snapshot = SnapshotTestWindow();
+				const FString ButtonRef = FindRefByRole(Snapshot, TEXT("button"), TEXT("Test Button"));
+				ASSERT_THAT(IsTrue(!ButtonRef.IsEmpty()));
+
+				const bool bHovered = USlateInspectorToolset::Hover(ButtonRef);
+				ASSERT_THAT(IsTrue(bHovered));
+			});
+	}
+
+	TEST_METHOD(Hover_InvalidRef)
+	{
+		const bool bHovered = USlateInspectorToolset::Hover(TEXT("nonexistent_ref"));
+		ASSERT_THAT(IsTrue(!bHovered));
+	}
+
+	// ── PressKey tests ──────────────────────────────────────────────
+
+	TEST_METHOD(PressKey_Escape)
+	{
+		const bool bPressed = USlateInspectorToolset::PressKey(TEXT("Escape"));
+		ASSERT_THAT(IsTrue(bPressed));
+	}
+
+	TEST_METHOD(PressKey_Tab)
+	{
+		const bool bPressed = USlateInspectorToolset::PressKey(TEXT("Tab"));
+		ASSERT_THAT(IsTrue(bPressed));
+	}
+
+	// ── SelectOption tests ──────────────────────────────────────────
+
+	TEST_METHOD(SelectOption_ChangesCombobox)
+	{
+		TestCommandBuilder
+			.StartWhen([this]() { return SlateComponent->HaveTicksElapsed(3); })
+			.Then([this]()
+			{
+				ASSERT_THAT(IsTrue(*TestPanel->SelectedComboOption == TEXT("Alpha")));
+
+				const FString Snapshot = SnapshotTestWindow();
+				const FString ComboRef = FindRefByRole(Snapshot, TEXT("combobox"), TEXT(""));
+				ASSERT_THAT(IsTrue(!ComboRef.IsEmpty()));
+
+				const bool bSelected = USlateInspectorToolset::SelectOption(ComboRef, TEXT("Beta"));
+				ASSERT_THAT(IsTrue(bSelected));
+				ASSERT_THAT(IsTrue(*TestPanel->SelectedComboOption == TEXT("Beta")));
+			});
+	}
+
+	TEST_METHOD(SelectOption_InvalidOption)
+	{
+		TestCommandBuilder
+			.StartWhen([this]() { return SlateComponent->HaveTicksElapsed(3); })
+			.Then([this]()
+			{
+				const FString Snapshot = SnapshotTestWindow();
+				const FString ComboRef = FindRefByRole(Snapshot, TEXT("combobox"), TEXT(""));
+				ASSERT_THAT(IsTrue(!ComboRef.IsEmpty()));
+
+				const bool bSelected = USlateInspectorToolset::SelectOption(ComboRef, TEXT("Nonexistent"));
+				ASSERT_THAT(IsTrue(!bSelected));
+				// Original selection unchanged.
+				ASSERT_THAT(IsTrue(*TestPanel->SelectedComboOption == TEXT("Alpha")));
+			});
+	}
+
+	// ── FillForm tests ──────────────────────────────────────────────
+
+	TEST_METHOD(FillForm_TextBox)
+	{
+		const FString Snapshot = SnapshotTestWindow();
+		const FString TextBoxRef = FindRefByRole(Snapshot, TEXT("textbox"), TEXT(""));
+		ASSERT_THAT(IsTrue(!TextBoxRef.IsEmpty()));
+
+		TArray<FSlateInspectorToolsetFormField> Fields;
+		FSlateInspectorToolsetFormField Field;
+		Field.Ref = TextBoxRef;
+		Field.Value = TEXT("form value");
+		Field.FieldType = TEXT("textbox");
+		Fields.Add(Field);
+
+		const bool bFilled = USlateInspectorToolset::FillForm(Fields);
+		ASSERT_THAT(IsTrue(bFilled));
+		ASSERT_THAT(IsTrue(TestPanel->TextBox->GetText().ToString() == TEXT("form value")));
+	}
+
+	TEST_METHOD(FillForm_Checkbox)
+	{
+		TestCommandBuilder
+			.StartWhen([this]() { return SlateComponent->HaveTicksElapsed(3); })
+			.Then([this]()
+			{
+				const FString Snapshot = SnapshotTestWindow();
+				const FString CheckBoxRef = FindRefByRole(Snapshot, TEXT("checkbox"), TEXT(""));
+				ASSERT_THAT(IsTrue(!CheckBoxRef.IsEmpty()));
+				ASSERT_THAT(IsTrue(!TestPanel->CheckBox->IsChecked()));
+
+				TArray<FSlateInspectorToolsetFormField> Fields;
+				FSlateInspectorToolsetFormField Field;
+				Field.Ref = CheckBoxRef;
+				Field.Value = TEXT("true");
+				Field.FieldType = TEXT("checkbox");
+				Fields.Add(Field);
+
+				const bool bFilled = USlateInspectorToolset::FillForm(Fields);
+				ASSERT_THAT(IsTrue(bFilled));
+				ASSERT_THAT(IsTrue(TestPanel->CheckBox->IsChecked()));
+			});
+	}
+
+	TEST_METHOD(FillForm_CheckboxAlreadyInDesiredState)
+	{
+		TestCommandBuilder
+			.StartWhen([this]() { return SlateComponent->HaveTicksElapsed(3); })
+			.Then([this]()
+			{
+				const FString Snapshot = SnapshotTestWindow();
+				const FString CheckBoxRef = FindRefByRole(Snapshot, TEXT("checkbox"), TEXT(""));
+				ASSERT_THAT(IsTrue(!TestPanel->CheckBox->IsChecked()));
+
+				TArray<FSlateInspectorToolsetFormField> Fields;
+				FSlateInspectorToolsetFormField Field;
+				Field.Ref = CheckBoxRef;
+				Field.Value = TEXT("false");
+				Field.FieldType = TEXT("checkbox");
+				Fields.Add(Field);
+
+				// Should succeed without clicking (already unchecked).
+				const bool bFilled = USlateInspectorToolset::FillForm(Fields);
+				ASSERT_THAT(IsTrue(bFilled));
+				ASSERT_THAT(IsTrue(!TestPanel->CheckBox->IsChecked()));
+			});
+	}
+
+	// ── WaitFor tests ───────────────────────────────────────────────
+
+	TEST_METHOD(WaitFor_FindsExistingText)
+	{
+		const bool bFound = USlateInspectorToolset::WaitFor(TEXT("Static Label"));
+		ASSERT_THAT(IsTrue(bFound));
+	}
+
+	TEST_METHOD(WaitFor_DoesNotFindMissingText)
+	{
+		const bool bFound = USlateInspectorToolset::WaitFor(TEXT("This Text Does Not Exist"));
+		ASSERT_THAT(IsTrue(!bFound));
+	}
+
+	TEST_METHOD(WaitFor_TextGone)
+	{
+		const bool bGone = USlateInspectorToolset::WaitFor(TEXT(""), TEXT("This Text Does Not Exist"));
+		ASSERT_THAT(IsTrue(bGone));
+	}
+
+	TEST_METHOD(WaitFor_TextGoneFails)
+	{
+		const bool bGone = USlateInspectorToolset::WaitFor(TEXT(""), TEXT("Static Label"));
+		ASSERT_THAT(IsTrue(!bGone));
+	}
+
+	// ── Windows tests ───────────────────────────────────────────────
+
+	TEST_METHOD(Windows_ListContainsTestWindow)
+	{
+		const FString WindowList = USlateInspectorToolset::Windows(TEXT("list"));
+		ASSERT_THAT(IsTrue(WindowList.Contains(TEXT("SlateInspectorToolset Test Window"))));
+	}
+
+	// ── Ref cache tests ─────────────────────────────────────────────
+
+	TEST_METHOD(PurgeExpired_CleansDeadRefs)
+	{
+		TSharedRef<STextBlock> EphemeralWidget = SNew(STextBlock).Text(FText::FromString(TEXT("Ephemeral")));
+		const FString Ref = FSlateInspectorToolsetRefCache::Get().GetOrAssignRef(EphemeralWidget, TEXT("x"));
+		ASSERT_THAT(IsTrue(!Ref.IsEmpty()));
+
+		TSharedPtr<SWidget> Resolved = FSlateInspectorToolsetRefCache::Get().ResolveRef(Ref);
+		ASSERT_THAT(IsTrue(Resolved.IsValid()));
+
+		// Drop all references so the widget is destroyed.
+		Resolved.Reset();
+		EphemeralWidget = SNew(STextBlock);
+
+		FSlateInspectorToolsetRefCache::Get().PurgeExpired();
+		TSharedPtr<SWidget> ResolvedAfterPurge = FSlateInspectorToolsetRefCache::Get().ResolveRef(Ref);
+		ASSERT_THAT(IsTrue(!ResolvedAfterPurge.IsValid()));
+	}
+
+	TEST_METHOD(FindRef_ReturnsEmptyForUnknownWidget)
+	{
+		TSharedRef<STextBlock> Widget = SNew(STextBlock);
+		const FString Ref = FSlateInspectorToolsetRefCache::Get().FindRef(Widget);
+		ASSERT_THAT(IsTrue(Ref.IsEmpty()));
+	}
+
+	TEST_METHOD(ResolveRef_ReturnsNullForUnknownRef)
+	{
+		TSharedPtr<SWidget> Resolved = FSlateInspectorToolsetRefCache::Get().ResolveRef(TEXT("zzz999"));
+		ASSERT_THAT(IsTrue(!Resolved.IsValid()));
+	}
+
+	// ── Helpers ─────────────────────────────────────────────────────
+
+	/** Snapshot only the test window's subtree so we don't accidentally
+	 *  match editor widgets (e.g. editor checkboxes vs. our test checkbox). */
+	FString SnapshotTestWindow()
+	{
+		// First get the window ref from a full snapshot.
+		const FString FullSnapshot = USlateInspectorToolset::Snapshot(TEXT(""), 30);
+		const FString WindowRef = FindRefByRole(FullSnapshot, TEXT("window"), TEXT("SlateInspectorToolset Test Window"));
+		if (WindowRef.IsEmpty())
+		{
+			return FString();
+		}
+		return USlateInspectorToolset::Snapshot(WindowRef, 30);
+	}
+
+	FString FindRefByRole(const FString& Snapshot, const FString& Role, const FString& Label)
+	{
+		TArray<FString> Lines;
+		Snapshot.ParseIntoArrayLines(Lines);
+		for (const FString& Line : Lines)
+		{
+			FString Trimmed = Line.TrimStartAndEnd();
+			if (!Trimmed.StartsWith(Role))
+			{
+				continue;
+			}
+			if (!Label.IsEmpty() && !Trimmed.Contains(Label))
+			{
+				continue;
+			}
+			int32 RefStart = Trimmed.Find(TEXT("[ref="));
+			if (RefStart != INDEX_NONE)
+			{
+				RefStart += 5;
+				int32 RefEnd = Trimmed.Find(TEXT("]"), ESearchCase::IgnoreCase, ESearchDir::FromStart, RefStart);
+				if (RefEnd != INDEX_NONE)
+				{
+					return Trimmed.Mid(RefStart, RefEnd - RefStart);
+				}
+			}
+		}
+		return FString();
+	}
+
+	TArray<FString> FindAllRefsByRole(const FString& Snapshot, const FString& Role)
+	{
+		TArray<FString> Refs;
+		TArray<FString> Lines;
+		Snapshot.ParseIntoArrayLines(Lines);
+		for (const FString& Line : Lines)
+		{
+			FString Trimmed = Line.TrimStartAndEnd();
+			if (!Trimmed.StartsWith(Role))
+			{
+				continue;
+			}
+			int32 RefStart = Trimmed.Find(TEXT("[ref="));
+			if (RefStart != INDEX_NONE)
+			{
+				RefStart += 5;
+				int32 RefEnd = Trimmed.Find(TEXT("]"), ESearchCase::IgnoreCase, ESearchDir::FromStart, RefStart);
+				if (RefEnd != INDEX_NONE)
+				{
+					Refs.Add(Trimmed.Mid(RefStart, RefEnd - RefStart));
+				}
+			}
+		}
+		return Refs;
+	}
+};
diff --git a/Engine/Plugins/Experimental/SlateInspectorToolset/Source/SlateInspectorToolsetTests/Private/SlateInspectorToolsetTestsModule.cpp b/Engine/Plugins/Experimental/SlateInspectorToolset/Source/SlateInspectorToolsetTests/Private/SlateInspectorToolsetTestsModule.cpp
new file mode 100644
--- /dev/null
+++ b/Engine/Plugins/Experimental/SlateInspectorToolset/Source/SlateInspectorToolsetTests/Private/SlateInspectorToolsetTestsModule.cpp
@@ -0,0 +1,9 @@
+// Copyright Epic Games, Inc. All Rights Reserved.
+
+#include "Modules/ModuleManager.h"
+
+class FSlateInspectorToolsetTestsModule : public IModuleInterface
+{
+};
+
+IMPLEMENT_MODULE(FSlateInspectorToolsetTestsModule, SlateInspectorToolsetTests)
diff --git a/Engine/Plugins/Experimental/SlateInspectorToolset/Source/SlateInspectorToolsetTests/SlateInspectorToolsetTests.Build.cs b/Engine/Plugins/Experimental/SlateInspectorToolset/Source/SlateInspectorToolsetTests/SlateInspectorToolsetTests.Build.cs
new file mode 100644
--- /dev/null
+++ b/Engine/Plugins/Experimental/SlateInspectorToolset/Source/SlateInspectorToolsetTests/SlateInspectorToolsetTests.Build.cs
@@ -0,0 +1,23 @@
+// Copyright Epic Games, Inc. All Rights Reserved.
+
+using UnrealBuildTool;
+
+public class SlateInspectorToolsetTests : ModuleRules
+{
+	public SlateInspectorToolsetTests(ReadOnlyTargetRules Target) : base(Target)
+	{
+		PCHUsage = PCHUsageMode.UseExplicitOrSharedPCHs;
+
+		PrivateDependencyModuleNames.AddRange(new string[]
+		{
+			"Core",
+			"CoreUObject",
+			"CQTest",
+			"Engine",
+			"InputCore",
+			"Slate",
+			"SlateCore",
+			"SlateInspectorToolset",
+		});
+	}
+}
