[SlateInspectorToolset] Move `SlateInspectorToolset` tests from `Editor` to `AI.Toolsets` category.

#Synced-CL 52658100

diff --git a/Engine/Plugins/Experimental/Toolsets/SlateInspectorToolset/Source/SlateInspectorToolsetTests/Private/SlateInspectorToolsetTests.cpp b/Engine/Plugins/Experimental/Toolsets/SlateInspectorToolset/Source/SlateInspectorToolsetTests/Private/SlateInspectorToolsetTests.cpp
--- a/Engine/Plugins/Experimental/Toolsets/SlateInspectorToolset/Source/SlateInspectorToolsetTests/Private/SlateInspectorToolsetTests.cpp
+++ b/Engine/Plugins/Experimental/Toolsets/SlateInspectorToolset/Source/SlateInspectorToolsetTests/Private/SlateInspectorToolsetTests.cpp
@@ -10,7 +10,7 @@
 #include "SlateInspectorToolset.h"
 #include "Widgets/SWindow.h"
 
-TEST_CLASS(SlateInspectorToolsetTest, "Editor.SlateInspectorToolset")
+TEST_CLASS(SlateInspectorToolsetTest, "AI.Toolsets.SlateInspectorToolset")
 {
 	TUniquePtr<FCQTestSlateComponent> SlateComponent;
 	TSharedPtr<SWindow> TestWindow;
